PVS-Studio noticed unintentional return statements (would have worked fine for provided input)

diff --git a/Engine/Source/Runtime/UnrealAudio/Private/Tests/UnrealAudioTestGenerators.cpp b/Engine/Source/Runtime/UnrealAudio/Private/Tests/UnrealAudioTestGenerators.cpp
--- a/Engine/Source/Runtime/UnrealAudio/Private/Tests/UnrealAudioTestGenerators.cpp
+++ b/Engine/Source/Runtime/UnrealAudio/Private/Tests/UnrealAudioTestGenerators.cpp
@@ -44,16 +44,7 @@ namespace Test
 
 	static float WrapTwoPi(float Value)
 	{
-		while (Value >= TWO_PI)
-		{
-			return Value -= TWO_PI;
-		}
-
-		while (Value < 0)
-		{
-			return Value += TWO_PI;
-		}
-		return Value;
+		return FMath::UnwindRadians(Value) + PI;
 	}
 
 	static float Lerp(float StartX, float EndX, float StartY, float EndY, float Alpha)
