UnrealAudio deleting commented code

diff --git a/Engine/Source/Runtime/UnrealAudio/Private/UnrealAudioDevice.cpp b/Engine/Source/Runtime/UnrealAudio/Private/UnrealAudioDevice.cpp
--- a/Engine/Source/Runtime/UnrealAudio/Private/UnrealAudioDevice.cpp
+++ b/Engine/Source/Runtime/UnrealAudio/Private/UnrealAudioDevice.cpp
@@ -100,24 +100,6 @@ namespace UAudio
 		return new FUnrealAudioDeviceDummy();
 	}
 
-// 	void OnDeviceError(const EDeviceError::Type Error, const FString& ErrorDetails, const FString& FileName, int32 LineNumber)
-// 	{
-// 		UE_LOG(LogUnrealAudioDevice, Error, TEXT("Audio Device Error: (%s) : %s (%s::%d)"), EDeviceError::ToString(Error), *ErrorDetails, *FileName, LineNumber);
-// 	}
-
-// 	uint32 GetNumBytesForFormat(EStreamFormat::Type Format)
-// 	{
-// 		switch (Format)
-// 		{
-// 			case EStreamFormat::FLT:		return 4;
-// 			case EStreamFormat::DBL:		return 8;
-// 			case EStreamFormat::INT_16:		return 2;
-// 			case EStreamFormat::INT_24:		return 3;
-// 			case EStreamFormat::INT_32:		return 4;
-// 			default:						return 0;
-// 		}
-// 	}
-
 }
 
 #endif // #if ENABLE_UNREAL_AUDIO
