UnrealAudio Soundfile encoding and decoding

- Implemented libsndfile soundfile class for encoding and decoding audio files (.wav, .aif, .ogg, .flac)
- Implemented import/export process tasks as async background tasks
- Added testing code to test single-file importing and exporting
- Added testing code to test batch-file importing and exporting
- Add dlls for Win32/Win64/Mac. Tested on Win64 and Mac.

To test file import/export processing, use the AudioTestCommandlet with the following:
-command UnrealEd.AudioTestCommandlet source export [SingleFile/Directory]

Where [SingleFile/Directory] is the full path to either a single audio file or a directory to a folder containing audio assets supported. All files are imported into an internal binary ogg format (I will make it variable import format to test), then immediately exported back to the original format. If a file fails importing due to a file format issue (or some other issue), that issue will be reported in the log and the batch processing will proceed.

The exported file will be exported into the original format and put in a Exported sub-directory of the directory of the given single file or the import directory with "_exported" appended to the name.

diff --git a/Engine/Source/Editor/UnrealEd/Private/Commandlets/AudioTestCommandlet.cpp b/Engine/Source/Editor/UnrealEd/Private/Commandlets/AudioTestCommandlet.cpp
--- a/Engine/Source/Editor/UnrealEd/Private/Commandlets/AudioTestCommandlet.cpp
+++ b/Engine/Source/Editor/UnrealEd/Private/Commandlets/AudioTestCommandlet.cpp
@@ -39,7 +39,7 @@ static bool UnrealAudioUnload()
 * Test static functions which call into module test code
 */
 
-static bool TestAudioDeviceAll()
+static bool TestAudioDeviceAll(const TArray<FString>& Args)
 {
 	if (!UAudio::TestDeviceQuery())
 	{
@@ -64,26 +64,60 @@ static bool TestAudioDeviceAll()
 	return true;
 }
 
-static bool TestAudioDeviceQuery()
+static bool TestAudioDeviceQuery(const TArray<FString>& Args)
 {
 	return UAudio::TestDeviceQuery();
 }
 
-static bool TestAudioDeviceOutputSimple()
+static bool TestAudioDeviceOutputSimple(const TArray<FString>& Args)
 {
 	return UAudio::TestDeviceOutputSimple(-1);
 }
 
-static bool TestAudioDeviceOutputFm()
+static bool TestAudioDeviceOutputFm(const TArray<FString>& Args)
 {
 	return UAudio::TestDeviceOutputRandomizedFm(-1);
 }
 
-static bool TestAudioDeviceOutputPan()
+static bool TestAudioDeviceOutputPan(const TArray<FString>& Args)
 {
 	return UAudio::TestDeviceOutputNoisePan(-1);
 }
 
+static bool TestAudioSourceImport(const TArray<FString>& Args)
+{
+	if (Args.Num() != 1)
+	{
+		return false;
+	}
+
+	UAudio::FSoundFileImportSettings SoundFileImportSettings;
+	SoundFileImportSettings.SoundFilePath = Args[0];
+	SoundFileImportSettings.Format = UAudio::ESoundFileFormat::OGG | UAudio::ESoundFileFormat::VORBIS;
+	SoundFileImportSettings.SampleRate = 48000;
+	SoundFileImportSettings.EncodingQuality = 0.5;
+	SoundFileImportSettings.bPerformPeakNormalization = true;
+
+	return UAudio::TestSourceImport(SoundFileImportSettings);
+}
+
+static bool TestAudioSourceImportExport(const TArray<FString>& Args)
+{
+	if (Args.Num() != 1)
+	{
+		return false;
+	}
+
+	UAudio::FSoundFileImportSettings SoundFileImportSettings;
+	SoundFileImportSettings.SoundFilePath = Args[0];
+	SoundFileImportSettings.Format = UAudio::ESoundFileFormat::OGG | UAudio::ESoundFileFormat::VORBIS;
+	SoundFileImportSettings.SampleRate = 48000;
+	SoundFileImportSettings.EncodingQuality = 0.5;
+	SoundFileImportSettings.bPerformPeakNormalization = false;
+
+	return UAudio::TestSourceImportExport(SoundFileImportSettings);
+}
+
 /**
 * Setting up commandlet code
 */
@@ -94,7 +128,9 @@ struct AudioTestInfo
 {
 	FString CategoryName;
 	FString TestName;
-	bool (*TestFunction)();
+	FString ArgDescription;
+	int32 NumArgs;
+	bool (*TestFunction)(const TArray<FString>& Args);
 };
 
 enum EAudioTests
@@ -104,26 +140,32 @@ enum EAudioTests
 	AUDIO_TEST_DEVICE_OUTPUT_SIMPLE,
 	AUDIO_TEST_DEVICE_OUTPUT_FM,
 	AUDIO_TEST_DEVICE_OUTPUT_PAN,
+	AUDIO_TEST_SOURCE_IMPORT,
+	AUDIO_TEST_SOURCE_IMPORT_EXPORT,
 	AUDIO_TESTS
 };
 
 static AudioTestInfo AudioTestInfoList[] =
 {
-	{ "device", "all",			TestAudioDeviceAll			}, // AUDIO_TEST_DEVICE_ALL
-	{ "device", "query",		TestAudioDeviceQuery		}, // AUDIO_TEST_DEVICE_QUERY
-	{ "device", "out",			TestAudioDeviceOutputSimple }, // AUDIO_TEST_DEVICE_OUTPUT_SIMPLE
-	{ "device", "out_fm",		TestAudioDeviceOutputFm		}, // AUDIO_TEST_DEVICE_OUTPUT_FM
-	{ "device", "out_pan",		TestAudioDeviceOutputPan	}, // AUDIO_TEST_DEVICE_OUTPUT_PAN
+	{ "device", "all",		"None",						0, TestAudioDeviceAll			}, // AUDIO_TEST_DEVICE_ALL
+	{ "device", "query",	"None",						0, TestAudioDeviceQuery			}, // AUDIO_TEST_DEVICE_QUERY
+	{ "device", "out",		"None",						0, TestAudioDeviceOutputSimple	}, // AUDIO_TEST_DEVICE_OUTPUT_SIMPLE
+	{ "device", "out_fm",	"None",						0, TestAudioDeviceOutputFm		}, // AUDIO_TEST_DEVICE_OUTPUT_FM
+	{ "device", "out_pan",	"None",						0, TestAudioDeviceOutputPan		}, // AUDIO_TEST_DEVICE_OUTPUT_PAN
+	{ "source", "import",	"SourcePath",				1, TestAudioSourceImport		}, // AUDIO_TEST_SOURCE_IMPORT
+	{ "source", "export",	"SourcePath or Directory",	1, TestAudioSourceImportExport	}, // AUDIO_TEST_SOURCE_IMPORT_EXPORT
 };
 static_assert(ARRAY_COUNT(AudioTestInfoList) == AUDIO_TESTS, "Mismatched info list and test enumeration");
 
 static void PrintUsage()
 {
-	UE_LOG(AudioTestCommandlet, Display, TEXT("AudioTestCommandlet Usage: {Editor}.exe UnrealEd.AudioTestCommandlet {testcategory} {test}"));
-	UE_LOG(AudioTestCommandlet, Display, TEXT("Possible Tests: [testcategory]: [test]"));
+	UE_LOG(AudioTestCommandlet, Display, TEXT("AudioTestCommandlet Usage: {Editor}.exe UnrealEd.AudioTestCommandlet {testcategory} {test} {arglist}"));
+	UE_LOG(AudioTestCommandlet, Display, TEXT("Possible Tests:"));
+	UE_LOG(AudioTestCommandlet, Display, TEXT("CategoryName | TestName | Arguments"));
 	for (uint32 Index = 0; Index < AUDIO_TESTS; ++Index)
 	{
-		UE_LOG(AudioTestCommandlet, Display, TEXT("    %s: %s"), *AudioTestInfoList[Index].CategoryName, *AudioTestInfoList[Index].TestName);
+		AudioTestInfo& Info = AudioTestInfoList[Index];
+		UE_LOG(AudioTestCommandlet, Display, TEXT("%s, %s, %s"), *Info.CategoryName, *Info.TestName, *Info.ArgDescription);
 	}
 }
 
@@ -160,24 +202,15 @@ int32 UAudioTestCommandlet::Main(const FString& InParams)
 	TArray<FString> Switches;
 	UCommandlet::ParseCommandLine(*Params, Tokens, Switches);
 
-	if (Tokens.Num() < 3 || Tokens.Num() > 4)
-	{
-		PrintUsage();
-		return 0;
-	}
-	int32 CategoryNameIndex = 1;
-	int32 TestNameIndex = 2;
-
-	if (Tokens.Num() == 4)
+	// Check if we're switching to a different device API
+	if (Switches.Num() == 1)
 	{
-		FString DeviceApiName(Tokens[1]);
+		FString DeviceApiName(Switches[1]);
 		if (!UnrealAudioLoad(&DeviceApiName))
 		{
 			UE_LOG(AudioTestCommandlet, Display, TEXT("Failed to load unreal audio module. Exiting."));
 			return 0;
 		}
-		CategoryNameIndex++;
-		TestNameIndex++;
 	}
 	else
 	{
@@ -190,21 +223,36 @@ int32 UAudioTestCommandlet::Main(const FString& InParams)
 
 	check(UnrealAudioModule != nullptr);
 
+	if (Tokens.Num() < 3)
+	{
+		PrintUsage();
+		return 0;
+	}
+
+	const int32 CategoryNameIndex = 1;
+	const int32 TestNameIndex = 2;
+	const int32 ArgStartIndex = 3;
 	bool FoundTest = false;
 	for (uint32 Index = 0; Index < AUDIO_TESTS; ++Index)
 	{
-		if (AudioTestInfoList[Index].CategoryName == Tokens[CategoryNameIndex])
+		AudioTestInfo& Info = AudioTestInfoList[Index];
+		if (Info.CategoryName == Tokens[CategoryNameIndex])
 		{
-			if (AudioTestInfoList[Index].TestName == Tokens[TestNameIndex])
+			if (Info.TestName == Tokens[TestNameIndex])
 			{
 				FoundTest = true;
-				if (AudioTestInfoList[Index].TestFunction())
+				TArray<FString> Args;
+				for (int32 ArgIndex = ArgStartIndex; ArgIndex < Tokens.Num(); ++ArgIndex)
+				{
+					Args.Add(Tokens[ArgIndex]);
+				}
+				if (AudioTestInfoList[Index].TestFunction(Args))
 				{
-					UE_LOG(AudioTestCommandlet, Display, TEXT("Test %s succeeded."), *AudioTestInfoList[Index].TestName);
+					UE_LOG(AudioTestCommandlet, Display, TEXT("Test %s succeeded."), *Info.TestName);
 				}
 				else
 				{
-					UE_LOG(AudioTestCommandlet, Display, TEXT("Test %s failed."), *AudioTestInfoList[Index].TestName);
+					UE_LOG(AudioTestCommandlet, Display, TEXT("Test %s failed."), *Info.TestName);
 				}
 				break;
 			}
diff --git a/Engine/Source/Runtime/UnrealAudio/Private/Tests/UnrealAudioDeviceTests.cpp b/Engine/Source/Runtime/UnrealAudio/Private/Tests/UnrealAudioDeviceTests.cpp
--- a/Engine/Source/Runtime/UnrealAudio/Private/Tests/UnrealAudioDeviceTests.cpp
+++ b/Engine/Source/Runtime/UnrealAudio/Private/Tests/UnrealAudioDeviceTests.cpp
@@ -5,6 +5,7 @@
 #include "UnrealAudioDeviceModule.h"
 #include "UnrealAudioTests.h"
 #include "UnrealAudioTestGenerators.h"
+#include "UnrealAudioSoundFile.h"
 
 #if ENABLE_UNREAL_AUDIO
 
@@ -228,6 +229,206 @@ namespace UAudio
 		return DoOutputTest("output white noise panner", LifeTime, &SimpleWhiteNoisePan);
 	}
 
+	bool TestSourceImport(const FSoundFileImportSettings& ImportSettings)
+	{
+		UE_LOG(LogUnrealAudio, Display, TEXT("Testing exporting a single sound source."));
+
+		UE_LOG(LogUnrealAudio, Display, TEXT("First importing."));
+		UE_LOG(LogUnrealAudio, Display, TEXT("Import Path: %s"), *ImportSettings.SoundFilePath);
+		UE_LOG(LogUnrealAudio, Display, TEXT("Import Format: %s"), *FString::Printf(TEXT("%s - %s"), ESoundFileFormat::ToStringMajor(ImportSettings.Format), ESoundFileFormat::ToStringMinor(ImportSettings.Format)));
+		UE_LOG(LogUnrealAudio, Display, TEXT("Import SampleRate: %d"), ImportSettings.SampleRate);
+		UE_LOG(LogUnrealAudio, Display, TEXT("Import EncodingQuality: %.2f"), ImportSettings.EncodingQuality);
+		UE_LOG(LogUnrealAudio, Display, TEXT("Perform Peak Normalization: %s"), ImportSettings.bPerformPeakNormalization ? TEXT("Yes") : TEXT("No"));
+
+		
+		TSharedPtr<ISoundFile> ImportedSoundFile = UnrealAudioModule->ImportSound(ImportSettings);
+		ESoundFileState::Type State = ImportedSoundFile->GetState();
+		while (State != ESoundFileState::IMPORTED)
+		{
+			FPlatformProcess::Sleep(1);
+
+			State = ImportedSoundFile->GetState();
+			if (State == ESoundFileState::HAS_ERROR)
+			{
+				UE_LOG(LogUnrealAudio, Error, TEXT("Failed to import sound source!"));
+				UE_LOG(LogUnrealAudio, Error, TEXT("Error: %s"), ESoundFileError::ToString(ImportedSoundFile->GetError()));
+				return false;
+			}
+		}
+		UE_LOG(LogUnrealAudio, Display, TEXT("Succeeded importing sound source!"));
+		return true;
+	}
+
+	bool TestSourceImportExport(const FSoundFileImportSettings& ImportSettings)
+	{
+		//  Check if the soundfile path is a folder or a single file:
+		if (FPaths::DirectoryExists(ImportSettings.SoundFilePath))
+		{
+			FString FolderPath = ImportSettings.SoundFilePath;
+			UE_LOG(LogUnrealAudio, Display, TEXT("Testing import then export of all audio files in directory: %s."), *FolderPath);
+
+			// Create the output exported directory if it doesn't exist
+			FString OutputDir = FolderPath / "Exported";
+			if (!FPaths::DirectoryExists(OutputDir))
+			{
+				FPlatformFileManager::Get().GetPlatformFile().CreateDirectory(*OutputDir);
+			}
+
+			// Get list of .wav files:
+			TArray<FString> InputFiles;
+
+			// Get input files in formats we're currently supporting
+			IFileManager::Get().FindFiles(InputFiles, *(FolderPath / "*.wav"), true, false);
+			IFileManager::Get().FindFiles(InputFiles, *(FolderPath / "*.aif"), true, false);
+			IFileManager::Get().FindFiles(InputFiles, *(FolderPath / "*.flac"), true, false);
+			IFileManager::Get().FindFiles(InputFiles, *(FolderPath / "*.ogg"), true, false);
+
+			// Make an array of shared ptr sound files
+			TArray<TSharedPtr<ISoundFile>> SoundFiles;
+
+			// Array of sound files (names) that have finished importing
+			TSet<FString> FinishedImporting;
+
+			// Array of sound files (names) that have error'd out during import
+			TSet<FString> HasError;
+
+			FSoundFileImportSettings Settings = ImportSettings;
+
+			// Import all the files!
+			UE_LOG(LogUnrealAudio, Display, TEXT("Importing..."));
+			for (FString& InputFile : InputFiles)
+			{
+				Settings.SoundFilePath = FolderPath / InputFile;
+				UE_LOG(LogUnrealAudio, Display, TEXT("%s"), *InputFile);
+
+				// start the import process and store the ISoundFile ptr in our array of in-flight imports
+				TSharedPtr<ISoundFile> ImportedSoundFile = UnrealAudioModule->ImportSound(Settings);
+				SoundFiles.Add(ImportedSoundFile);
+			}
+
+			// Loop through in-flight imports and start exports as soon as the file finishes importing
+			bool bIsFinishedImporting = false;
+			while (!bIsFinishedImporting)
+			{
+				// Assume we're done, but set false as soon as we find an import still going
+				bIsFinishedImporting = true;
+
+				// Loop through our stored sound files
+				for (TSharedPtr<ISoundFile> SoundFile : SoundFiles)
+				{
+					// Get the sound file state, handle any errors that occurred during import
+					// Note: its ok for our test assets here to fail since we're testing graceful failure. We
+					// might have some pathological sound assets in the test asset directory.
+					ESoundFileState::Type State = SoundFile->GetState();
+					if (State == ESoundFileState::HAS_ERROR)
+					{
+						FString Name;
+						SoundFile->GetName(Name);
+
+						// Store the error in a list of error'd sound files
+						if (!HasError.Contains(Name))
+						{
+							HasError.Add(Name);
+						}
+					}
+					else if (State != ESoundFileState::IMPORTED)
+					{
+						bIsFinishedImporting = false;
+					}
+					else
+					{
+						FString Name;
+						SoundFile->GetName(Name);
+
+						// If this file's state is IMPORTED, then we add it to the list
+						// of FinishedImporting sounds, then trigger the export process.
+						if (!FinishedImporting.Contains(Name))
+						{
+							FinishedImporting.Add(Name);
+
+							UE_LOG(LogUnrealAudio, Display, TEXT("Finished importing %s"), *Name);
+
+							TSharedPtr<ISoundFileData> SoundFileData;
+							SoundFile->GetSoundFileData(SoundFileData);
+
+							// We're building up a filename with [NAME]_exported.[EXT] in the output directory
+							const FString& SoundFileFullPath = SoundFileData->GetFilePath();
+							FString BaseSoundFileName = FPaths::GetBaseFilename(SoundFileFullPath);
+							FString SoundFileExtension = FPaths::GetExtension(SoundFileFullPath);
+							FString SoundFilePath = FPaths::GetPath(SoundFileFullPath);
+							FString OutputFileName = BaseSoundFileName + TEXT("_exported.") + SoundFileExtension;
+							FString OutputPath = OutputDir / OutputFileName;
+
+							UE_LOG(LogUnrealAudio, Display, TEXT("Now exporting to %s"), *OutputPath);
+
+							UnrealAudioModule->ExportSound(SoundFileData, OutputPath);
+						}
+					}
+				}
+
+				// Sleep until all imports have finished. In-flight exports
+				// will cause the UnrealAudio shutdown to block until they are finished.
+				FPlatformProcess::Sleep(1);
+			}
+
+			return true;
+		}
+		else if (FPaths::FileExists(ImportSettings.SoundFilePath))
+		{
+			UE_LOG(LogUnrealAudio, Display, TEXT("Testing import, then export of a single sound source: %s."), *ImportSettings.SoundFilePath);
+			UE_LOG(LogUnrealAudio, Display, TEXT("Import Format: %s"), *FString::Printf(TEXT("%s - %s"), ESoundFileFormat::ToStringMajor(ImportSettings.Format), ESoundFileFormat::ToStringMinor(ImportSettings.Format)));
+			UE_LOG(LogUnrealAudio, Display, TEXT("Import SampleRate: %d"), ImportSettings.SampleRate);
+			UE_LOG(LogUnrealAudio, Display, TEXT("Import EncodingQuality: %.2f"), ImportSettings.EncodingQuality);
+			UE_LOG(LogUnrealAudio, Display, TEXT("Perform Peak Normalization: %s"), ImportSettings.bPerformPeakNormalization ? TEXT("Yes") : TEXT("No"));
+
+			// Start the import process on the sound
+			TSharedPtr<ISoundFile> ImportedSoundFile = UnrealAudioModule->ImportSound(ImportSettings);
+
+			// Wait for the import to finish (block this thread for test-purposes). Normally, this would be non-blocking.
+			ESoundFileState::Type State = ImportedSoundFile->GetState();
+			while (State != ESoundFileState::IMPORTED)
+			{
+				FPlatformProcess::Sleep(1);
+
+				State = ImportedSoundFile->GetState();
+				if (State == ESoundFileState::HAS_ERROR)
+				{
+					UE_LOG(LogUnrealAudio, Error, TEXT("Failed to import sound source!"));
+					UE_LOG(LogUnrealAudio, Error, TEXT("Error: %s"), ESoundFileError::ToString(ImportedSoundFile->GetError()));
+					return false;
+				}
+			}
+
+			UE_LOG(LogUnrealAudio, Display, TEXT("Succeeded importing sound source!"));
+
+			// Now setup the export path
+			FString SoundFileFullPath = ImportSettings.SoundFilePath;
+			FString BaseSoundFileName = FPaths::GetBaseFilename(SoundFileFullPath);
+			FString SoundFileExtension = FPaths::GetExtension(SoundFileFullPath);
+			FString SoundFileDir = FPaths::GetPath(SoundFileFullPath);
+
+			// Create the export directory if it doesn't exist
+			FString OutputDir = SoundFileDir;
+			if (!FPaths::DirectoryExists(OutputDir))
+			{
+				FPlatformFileManager::Get().GetPlatformFile().CreateDirectory(*OutputDir);
+			}
+
+			// Append _exported to the file path just to make it clear that this is the exported version of the file
+			FString OutputPath = SoundFileDir / (BaseSoundFileName + TEXT("_exported.") + SoundFileExtension);
+
+			// Start the export process. We can exit this test function since UnrealAudio will block on shutdown
+			// until all background processes have finished. Normally the Editor will remain running during the export process.
+			TSharedPtr<ISoundFileData> SoundFileData;
+			ImportedSoundFile->GetSoundFileData(SoundFileData);
+			UnrealAudioModule->ExportSound(SoundFileData, OutputPath);
+			return true;
+		}
+
+		UE_LOG(LogUnrealAudio, Display, TEXT("Path %s is not a single file or a directory."), *ImportSettings.SoundFilePath);
+		return false;
+	}
+
 }
 
 #endif // #if ENABLE_UNREAL_AUDIO
diff --git a/Engine/Source/Runtime/UnrealAudio/Private/UnrealAudio.cpp b/Engine/Source/Runtime/UnrealAudio/Private/UnrealAudio.cpp
--- a/Engine/Source/Runtime/UnrealAudio/Private/UnrealAudio.cpp
+++ b/Engine/Source/Runtime/UnrealAudio/Private/UnrealAudio.cpp
@@ -21,6 +21,7 @@ namespace UAudio
 {
 	FUnrealAudioModule::FUnrealAudioModule()
 		: UnrealAudioDevice(nullptr)
+		, SoundFileDllHandle(nullptr)
 	{
 	}
 
@@ -63,11 +64,27 @@ namespace UAudio
 		{
 			UnrealAudioDevice = CreateDummyDeviceModule();
 		}
+
+		bSuccess = LoadSoundFileLib();
+
 		return bSuccess;
 	}
 
 	void FUnrealAudioModule::Shutdown()
 	{
+		// Block shutdown until the number of background tasks goes to zero before shutting down the audio module
+		static const int32 BackgroundTaskTimeout = 500;
+		int32 BackgroundTaskWaitCount = 0;
+		while (NumBackgroundTasks.GetValue() != 0 && BackgroundTaskWaitCount++ < BackgroundTaskTimeout)
+		{
+			FPlatformProcess::Sleep(1);
+		}
+
+		if (BackgroundTaskWaitCount == BackgroundTaskTimeout)
+		{
+			UE_LOG(LogUnrealAudio, Error, TEXT("Timed out while waiting for background tasks to finish when shutting down unreal audio module."));
+		}
+
 		if (UnrealAudioDevice)
 		{
 			UnrealAudioDevice->Shutdown();
@@ -79,8 +96,9 @@ namespace UAudio
 			delete UnrealAudioDevice;
 			UnrealAudioDevice = nullptr;
 		}
-	}
 
+		ShutdownSoundFileLib();
+	}
 
 	FName FUnrealAudioModule::GetDefaultDeviceModuleName() const
 	{
@@ -94,6 +112,16 @@ namespace UAudio
 		return UnrealAudioDevice;
 	}
 
+	void FUnrealAudioModule::IncrementBackgroundTaskCount()
+	{
+		NumBackgroundTasks.Increment();
+	}
+
+	void FUnrealAudioModule::DecrementBackgroundTaskCount()
+	{
+		NumBackgroundTasks.Decrement();
+	}
+
 }
 
 #endif // #if ENABLE_UNREAL_AUDIO
diff --git a/Engine/Source/Runtime/UnrealAudio/Private/UnrealAudioPrivate.h b/Engine/Source/Runtime/UnrealAudio/Private/UnrealAudioPrivate.h
--- a/Engine/Source/Runtime/UnrealAudio/Private/UnrealAudioPrivate.h
+++ b/Engine/Source/Runtime/UnrealAudio/Private/UnrealAudioPrivate.h
@@ -12,29 +12,56 @@ DECLARE_LOG_CATEGORY_EXTERN(LogUnrealAudio, Log, All);
 
 namespace UAudio
 {
-	class FUnrealAudioModule :	public IUnrealAudioModule
+	/** 
+	* FUnrealAudioModule 
+	* Concrete implementation of IUnrealAudioModule
+	*/
+	class FUnrealAudioModule : public IUnrealAudioModule
 	{
 	public:
 
 		FUnrealAudioModule();
 		~FUnrealAudioModule();
 
+		// IUnrealaAudioModule Interface
 		bool Initialize() override;
 		bool Initialize(const FString& DeviceModuleName) override;
 		void Shutdown() override;
+		TSharedPtr<ISoundFile> ImportSound(const FSoundFileImportSettings& ImportSettings) override;
+		void ExportSound(TSharedPtr<ISoundFileData> SoundFileData, const FString& ExportPath) override;
 
 		class IUnrealAudioDeviceModule* GetDeviceModule();
-
 		FName GetDefaultDeviceModuleName() const;
 
+		// Internal API
+		void IncrementBackgroundTaskCount();
+		void DecrementBackgroundTaskCount();
+
 	private:
 
 		bool InitializeInternal();
+		
+		/** Loads the sound file parsing libs */
+		bool LoadSoundFileLib();
+		bool ShutdownSoundFileLib();
+
+		/** Initializes unreal audio tests */
 		static void InitializeTests(FUnrealAudioModule* Module);
 
 		class IUnrealAudioDeviceModule* UnrealAudioDevice;
 		FName ModuleName;
+
+		/** Number of background tasks in-flight */
+		FThreadSafeCounter NumBackgroundTasks;
+
+		/** Handle to the sound file DLL */
+		void* SoundFileDllHandle;
+
+
 	};
+
+
+
 }
 
 #endif // #if ENABLE_UNREAL_AUDIO
diff --git a/Engine/Source/Runtime/UnrealAudio/Private/UnrealAudioSoundFile.cpp b/Engine/Source/Runtime/UnrealAudio/Private/UnrealAudioSoundFile.cpp
new file mode 100644
--- /dev/null
+++ b/Engine/Source/Runtime/UnrealAudio/Private/UnrealAudioSoundFile.cpp
@@ -0,0 +1,918 @@
+#include "UnrealAudioPrivate.h"
+#include "UnrealAudioSoundFile.h"
+#include "UnrealAudioSoundFileInternal.h"
+#include "UnrealAudioModule.h"
+#include "ModuleManager.h"
+
+#if ENABLE_UNREAL_AUDIO
+
+namespace UAudio
+{
+	// Virtual Sound File Function Pointers
+	typedef SoundFileCount(*VirtualSoundFileGetLengthFuncPtr)(void* UserData);
+	typedef SoundFileCount(*VirtualSoundFileSeekFuncPtr)(SoundFileCount Offset, int32 Mode, void* UserData);
+	typedef SoundFileCount(*VirtualSoundFileReadFuncPtr)(void* DataPtr, SoundFileCount ByteCount, void* UserData);
+	typedef SoundFileCount(*VirtualSoundFileWriteFuncPtr)(const void* DataPtr, SoundFileCount ByteCount, void* UserData);
+	typedef SoundFileCount(*VirtualSoundFileTellFuncPtr)(void* UserData);
+
+	// Struct describing function pointers to call for virtual file IO
+	struct FVirtualSoundFileCallbackInfo
+	{
+		VirtualSoundFileGetLengthFuncPtr VirtualSoundFileGetLength;
+		VirtualSoundFileSeekFuncPtr VirtualSoundFileSeek;
+		VirtualSoundFileReadFuncPtr VirtualSoundFileRead;
+		VirtualSoundFileWriteFuncPtr VirtualSoundFileWrite;
+		VirtualSoundFileTellFuncPtr VirtualSoundFileTell;
+	};
+
+	// SoundFile Constants
+	static const int32 SET_ENCODING_QUALITY = 0x1300;
+	static const int32 SET_CHANNEL_MAP_INFO = 0x1101;
+	static const int32 GET_CHANNEL_MAP_INFO = 0x1100;
+
+	// Exported SoundFile Functions
+	typedef SoundFileHandle*(*SoundFileOpenFuncPtr)(const char* Path, int32 Mode, FSoundFileDescription* Description);
+	typedef SoundFileHandle*(*SoundFileOpenVirtualFuncPtr)(FVirtualSoundFileCallbackInfo* VirtualFileDescription, int32 Mode, FSoundFileDescription* Description, void* UserData);
+	typedef int32(*SoundFileCloseFuncPtr)(SoundFileHandle* FileHandle);
+	typedef int32(*SoundFileErrorFuncPtr)(SoundFileHandle* FileHandle);
+	typedef const char*(*SoundFileStrErrorFuncPtr)(SoundFileHandle* FileHandle);
+	typedef const char*(*SoundFileErrorNumberFuncPtr)(int32 ErrorNumber);
+	typedef int32(*SoundFileCommandFuncPtr)(SoundFileHandle* FileHandle, int32 Command, void* Data, int32 DataSize);
+	typedef int32(*SoundFileFormatCheckFuncPtr)(const FSoundFileDescription* Description);
+	typedef SoundFileCount(*SoundFileSeekFuncPtr)(SoundFileHandle* FileHandle, SoundFileCount NumFrames, int32 SeekMode);
+	typedef const char*(*SoundFileGetVersionFuncPtr)(void);
+	typedef SoundFileCount(*SoundFileReadFramesFloatFuncPtr)(SoundFileHandle* FileHandle, float* Buffer, SoundFileCount NumFrames);
+	typedef SoundFileCount(*SoundFileReadFramesDoubleFuncPtr)(SoundFileHandle* FileHandle, double* Buffer, SoundFileCount NumFrames);
+	typedef SoundFileCount(*SoundFileWriteFramesFloatFuncPtr)(SoundFileHandle* FileHandle, const float* Buffer, SoundFileCount NumFrames);
+	typedef SoundFileCount(*SoundFileWriteFramesDoubleFuncPtr)(SoundFileHandle* FileHandle, const double* Buffer, SoundFileCount NumFrames);
+	typedef SoundFileCount(*SoundFileReadSamplesFloatFuncPtr)(SoundFileHandle* FileHandle, float* Buffer, SoundFileCount NumSamples);
+	typedef SoundFileCount(*SoundFileReadSamplesDoubleFuncPtr)(SoundFileHandle* FileHandle, double* Buffer, SoundFileCount NumSamples);
+	typedef SoundFileCount(*SoundFileWriteSamplesFloatFuncPtr)(SoundFileHandle* FileHandle, const float* Buffer, SoundFileCount NumSamples);
+	typedef SoundFileCount(*SoundFileWriteSamplesDoubleFuncPtr)(SoundFileHandle* FileHandle, const double* Buffer, SoundFileCount NumSamples);
+
+	SoundFileOpenFuncPtr SoundFileOpen = nullptr;
+	SoundFileOpenVirtualFuncPtr SoundFileOpenVirtual = nullptr;
+	SoundFileCloseFuncPtr SoundFileClose = nullptr;
+	SoundFileErrorFuncPtr SoundFileError = nullptr;
+	SoundFileStrErrorFuncPtr SoundFileStrError = nullptr;
+	SoundFileErrorNumberFuncPtr SoundFileErrorNumber = nullptr;
+	SoundFileCommandFuncPtr SoundFileCommand = nullptr;
+	SoundFileFormatCheckFuncPtr SoundFileFormatCheck = nullptr;
+	SoundFileSeekFuncPtr SoundFileSeek = nullptr;
+	SoundFileGetVersionFuncPtr SoundFileGetVersion = nullptr;
+	SoundFileReadFramesFloatFuncPtr SoundFileReadFramesFloat = nullptr;
+	SoundFileReadFramesDoubleFuncPtr SoundFileReadFramesDouble = nullptr;
+	SoundFileWriteFramesFloatFuncPtr SoundFileWriteFramesFloat = nullptr;
+	SoundFileWriteFramesDoubleFuncPtr SoundFileWriteFramesDouble = nullptr;
+	SoundFileReadSamplesFloatFuncPtr SoundFileReadSamplesFloat = nullptr;
+	SoundFileReadSamplesDoubleFuncPtr SoundFileReadSamplesDouble = nullptr;
+	SoundFileWriteSamplesFloatFuncPtr SoundFileWriteSamplesFloat = nullptr;
+	SoundFileWriteSamplesDoubleFuncPtr SoundFileWriteSamplesDouble = nullptr;
+
+	/** 
+	Function implementations of virtual function callbacks
+	*/
+
+	static SoundFileCount OnSoundFileGetLengthBytes(void* UserData)
+	{
+		SoundFileCount Length = 0;
+		((ISoundFileVirtual*)UserData)->GetLengthBytes(Length);
+		return Length;
+	}
+
+	static SoundFileCount OnSoundFileSeekBytes(SoundFileCount Offset, int32 Mode, void* UserData)
+	{
+		SoundFileCount OutOffset = 0;
+		((ISoundFileVirtual*)UserData)->SeekBytes(Offset, (ESoundFileSeekMode::Type)Mode, OutOffset);
+		return OutOffset;
+	}
+
+	static SoundFileCount OnSoundFileReadBytes(void* DataPtr, SoundFileCount ByteCount, void* UserData)
+	{
+		SoundFileCount OutBytesRead = 0;
+		((ISoundFileVirtual*)UserData)->ReadBytes(DataPtr, ByteCount, OutBytesRead);
+		return OutBytesRead;
+	}
+
+	static SoundFileCount OnSoundFileWriteBytes(const void* DataPtr, SoundFileCount ByteCount, void* UserData)
+	{
+		SoundFileCount OutBytesWritten = 0;
+		((ISoundFileVirtual*)UserData)->WriteBytes(DataPtr, ByteCount, OutBytesWritten);
+		return OutBytesWritten;
+	}
+
+	static SoundFileCount OnSoundFileTell(void* UserData)
+	{
+		return ((ISoundFileVirtual*)UserData)->GetOffsetBytes();
+	}
+
+	/**
+	* Gets the default channel mapping for the given channel number
+	*/
+	static void GetDefaultMappingsForChannelNumber(int32 NumChannels, TArray<ESoundFileChannelMap::Type>& ChannelMap)
+	{
+		check(ChannelMap.Num() == NumChannels);
+
+		switch (NumChannels)
+		{
+			case 1:	// MONO
+				ChannelMap[0] = ESoundFileChannelMap::MONO;
+				break;
+
+			case 2:	// STEREO
+				ChannelMap[0] = ESoundFileChannelMap::LEFT;
+				ChannelMap[1] = ESoundFileChannelMap::RIGHT;
+				break;
+
+			case 3:	// 2.1
+				ChannelMap[0] = ESoundFileChannelMap::LEFT;
+				ChannelMap[1] = ESoundFileChannelMap::RIGHT;
+				ChannelMap[2] = ESoundFileChannelMap::LFE;
+				break;
+
+			case 4: // Quadraphonic
+				ChannelMap[0] = ESoundFileChannelMap::LEFT;
+				ChannelMap[1] = ESoundFileChannelMap::RIGHT;
+				ChannelMap[2] = ESoundFileChannelMap::BACK_LEFT;
+				ChannelMap[3] = ESoundFileChannelMap::BACK_RIGHT;
+				break;
+
+			case 5: // 5.0
+				ChannelMap[0] = ESoundFileChannelMap::LEFT;
+				ChannelMap[1] = ESoundFileChannelMap::RIGHT;
+				ChannelMap[2] = ESoundFileChannelMap::CENTER;
+				ChannelMap[3] = ESoundFileChannelMap::SIDE_LEFT;
+				ChannelMap[4] = ESoundFileChannelMap::SIDE_RIGHT;
+				break;
+
+			case 6: // 5.1
+				ChannelMap[0] = ESoundFileChannelMap::LEFT;
+				ChannelMap[1] = ESoundFileChannelMap::RIGHT;
+				ChannelMap[2] = ESoundFileChannelMap::CENTER;
+				ChannelMap[3] = ESoundFileChannelMap::LFE;
+				ChannelMap[4] = ESoundFileChannelMap::SIDE_LEFT;
+				ChannelMap[5] = ESoundFileChannelMap::SIDE_RIGHT;
+				break;
+
+			case 7: // 6.1
+				ChannelMap[0] = ESoundFileChannelMap::LEFT;
+				ChannelMap[1] = ESoundFileChannelMap::RIGHT;
+				ChannelMap[2] = ESoundFileChannelMap::CENTER;
+				ChannelMap[3] = ESoundFileChannelMap::LFE;
+				ChannelMap[4] = ESoundFileChannelMap::SIDE_LEFT;
+				ChannelMap[5] = ESoundFileChannelMap::SIDE_RIGHT;
+				ChannelMap[6] = ESoundFileChannelMap::BACK_CENTER;
+				break;
+
+			case 8: // 7.1
+				ChannelMap[0] = ESoundFileChannelMap::LEFT;
+				ChannelMap[1] = ESoundFileChannelMap::RIGHT;
+				ChannelMap[2] = ESoundFileChannelMap::CENTER;
+				ChannelMap[3] = ESoundFileChannelMap::LFE;
+				ChannelMap[4] = ESoundFileChannelMap::BACK_LEFT;
+				ChannelMap[5] = ESoundFileChannelMap::BACK_RIGHT;
+				ChannelMap[6] = ESoundFileChannelMap::SIDE_LEFT;
+				ChannelMap[7] = ESoundFileChannelMap::SIDE_RIGHT;
+				break;
+
+			default:
+				break;
+		}
+	}
+
+	static int32 SoundFileId = 0;
+
+	FSoundFile::FSoundFile()
+		: CurrentIndexBytes(0)
+		, FileHandle(nullptr)
+		, State(ESoundFileState::UNINITIALIZED)
+		, CurrentError(ESoundFileError::NONE)
+		, Id(SoundFileId++)
+	{
+	}
+
+	FSoundFile::~FSoundFile()
+	{
+		check(FileHandle == nullptr);
+	}
+
+	ESoundFileError::Type FSoundFile::GetSoundFileData(TSharedPtr<ISoundFileData>& OutSoundFileData)
+	{
+		if (State == ESoundFileState::IMPORTING)
+		{
+			return SetError(ESoundFileError::ALREADY_INTIALIZED);
+		}
+
+		OutSoundFileData = SoundFileData;
+		return ESoundFileError::NONE;
+	}
+
+	ESoundFileError::Type FSoundFile::Initialize(TSharedPtr<ISoundFileData> InSoundFileData)
+	{
+		if (State != ESoundFileState::UNINITIALIZED)
+		{
+			return SetError(ESoundFileError::ALREADY_INTIALIZED);
+		}
+
+		check(FileHandle == nullptr);
+
+		// Setting sound file data initializes this sound file
+		SoundFileData = InSoundFileData;
+
+		// Open up a virtual file handle with this data
+		FVirtualSoundFileCallbackInfo VirtualSoundFileInfo;
+		VirtualSoundFileInfo.VirtualSoundFileGetLength = OnSoundFileGetLengthBytes;
+		VirtualSoundFileInfo.VirtualSoundFileSeek = OnSoundFileSeekBytes;
+		VirtualSoundFileInfo.VirtualSoundFileRead = OnSoundFileReadBytes;
+		VirtualSoundFileInfo.VirtualSoundFileWrite = OnSoundFileWriteBytes;
+		VirtualSoundFileInfo.VirtualSoundFileTell = OnSoundFileTell;
+
+		FSoundFileDescription Description = SoundFileData->GetDescription();
+		if (!SoundFileFormatCheck(&Description))
+		{
+			return SetError(ESoundFileError::INVALID_INPUT_FORMAT);
+		}
+
+		FileHandle = SoundFileOpenVirtual(&VirtualSoundFileInfo, ESoundFileOpenMode::READING, &Description, (void*)this);
+		if (!FileHandle)
+		{
+			FString StrErr = SoundFileStrError(nullptr);
+			UE_LOG(LogUnrealAudio, Error, TEXT("Failed to intitialize sound file: %s"), *StrErr);
+			return SetError(ESoundFileError::FAILED_TO_OPEN);
+		}
+
+		State = ESoundFileState::INITIALIZED;
+		return ESoundFileError::NONE;
+	}
+
+	ESoundFileError::Type FSoundFile::GetName(FString& Name)
+	{
+		if (SoundFileData.IsValid())
+		{
+			const FString& FilePath = SoundFileData->GetFilePath();
+			Name = FPaths::GetBaseFilename(FilePath);
+			return ESoundFileError::NONE;
+		}
+		return ESoundFileError::INVALID_DATA;
+	}
+
+	ESoundFileState::Type FSoundFile::GetState() const
+	{
+		return State;
+	}
+
+	ESoundFileError::Type FSoundFile::GetError() const
+	{
+		return CurrentError;
+	}
+
+	ESoundFileError::Type FSoundFile::SetError(ESoundFileError::Type Error)
+	{
+		if (Error != ESoundFileError::NONE)
+		{
+			State = ESoundFileState::HAS_ERROR;
+		}
+		return (CurrentError = Error);
+	}
+
+	ESoundFileError::Type FSoundFile::OpenFileForReading(const FString& FilePath)
+	{
+		if (FileHandle != nullptr)
+		{
+			UE_LOG(LogUnrealAudio, Error, TEXT("Sound file already opened"));
+			return SetError(ESoundFileError::ALREADY_OPENED);
+		}
+
+		// Make sure we don't already have data
+		if (SoundFileData.IsValid())
+		{
+			return SetError(ESoundFileError::ALREADY_HAS_DATA);
+		}
+
+		// Make a data ptr
+		SoundFileData = TSharedPtr<FSoundFileData>(new FSoundFileData(ESoundFileOpenMode::READING));
+		SoundFileData->SetFilePath(FilePath);
+
+		// Check to see if the file exists
+		if (!FPaths::FileExists(FilePath))
+		{
+			UE_LOG(LogUnrealAudio, Error, TEXT("Sound file %s doesn't exist."), *FilePath);
+			return SetError(ESoundFileError::FILE_DOESNT_EXIST);
+		}
+
+		// Open up the file
+		FSoundFileDescription Description;
+		FileHandle = SoundFileOpen(TCHAR_TO_ANSI(*FilePath), ESoundFileOpenMode::READING, &Description);
+		if (!FileHandle)
+		{
+			FString StrError = FString(SoundFileStrError(nullptr));
+			UE_LOG(LogUnrealAudio, Error, TEXT("Failed to open sound file %s: %s"), *FilePath, *StrError);
+			return SetError(ESoundFileError::FAILED_TO_OPEN);
+		}
+		SoundFileData->SetDescription(Description);
+
+		// Try to get a channel mapping
+		int32 NumChannels = Description.NumChannels;
+		TArray<ESoundFileChannelMap::Type> ChannelMap;
+		ChannelMap.Init(ESoundFileChannelMap::INVALID, NumChannels);
+
+		int32 Result = SoundFileCommand(FileHandle, GET_CHANNEL_MAP_INFO, (int32*)ChannelMap.GetData(), sizeof(int32)*NumChannels);
+
+		// If we failed to get the file's channel map definition, then we set the default based on the number of channels
+		if (Result == 0)
+		{
+			GetDefaultMappingsForChannelNumber(NumChannels, ChannelMap);
+		}
+		else
+		{
+			// Check to see if the channel map we did get back is filled with INVALID channels
+			bool bIsInvalid = false;
+			for (ESoundFileChannelMap::Type ChannelType : ChannelMap)
+			{
+				if (ChannelType == ESoundFileChannelMap::INVALID)
+				{
+					bIsInvalid = true;
+					break;
+				}
+			}
+			// If invalid, then we need to get the default channel mapping
+			if (bIsInvalid)
+			{
+				GetDefaultMappingsForChannelNumber(NumChannels, ChannelMap);
+			}
+		}
+		SoundFileData->SetChannelMap(ChannelMap);
+
+		return ESoundFileError::NONE;
+	}
+
+	ESoundFileError::Type FSoundFile::OpenEmptyFileForImport(const FSoundFileDescription& InDescription, const TArray<ESoundFileChannelMap::Type>& InChannelMap)
+	{
+		// First check the input format to make sure it's valid
+		if (!SoundFileFormatCheck(&InDescription))
+		{
+			UE_LOG(LogUnrealAudio, 
+				   Error, 
+				   TEXT("Sound file input format (%s - %s) is invalid."), 
+				   ESoundFileFormat::ToStringMajor(InDescription.FormatFlags), 
+				   ESoundFileFormat::ToStringMinor(InDescription.FormatFlags));
+
+			return SetError(ESoundFileError::INVALID_INPUT_FORMAT);
+		}
+		
+		// Make sure we don't already have data
+		if (SoundFileData.IsValid())
+		{
+			return SetError(ESoundFileError::ALREADY_HAS_DATA);
+		}
+		
+		// Make sure we have the right number of channels and our channel map size
+		if (InChannelMap.Num() != InDescription.NumChannels)
+		{
+			UE_LOG(LogUnrealAudio, Error, TEXT("Channel map didn't match the input NumChannels"));
+			return (CurrentError = ESoundFileError::INVALID_CHANNEL_MAP);
+		}
+		
+		// Make a data ptr
+		SoundFileData = TSharedPtr<FSoundFileData>(new FSoundFileData(ESoundFileOpenMode::WRITING));
+
+		FVirtualSoundFileCallbackInfo VirtualSoundFileInfo;
+		VirtualSoundFileInfo.VirtualSoundFileGetLength = OnSoundFileGetLengthBytes;
+		VirtualSoundFileInfo.VirtualSoundFileSeek = OnSoundFileSeekBytes;
+		VirtualSoundFileInfo.VirtualSoundFileRead = OnSoundFileReadBytes;
+		VirtualSoundFileInfo.VirtualSoundFileWrite = OnSoundFileWriteBytes;
+		VirtualSoundFileInfo.VirtualSoundFileTell = OnSoundFileTell;
+
+		FSoundFileDescription Description = InDescription;
+		FileHandle = SoundFileOpenVirtual(&VirtualSoundFileInfo, ESoundFileOpenMode::WRITING, &Description, (void*)this);
+		if (!FileHandle)
+		{
+			FString StrErr = FString(SoundFileStrError(nullptr));
+			UE_LOG(LogUnrealAudio, Error, TEXT("Failed to open empty sound file: %s"), *StrErr);
+			return SetError(ESoundFileError::FAILED_TO_OPEN);
+		}
+
+		// For some reason, virtual sound files set the NumFrames to be the max int64 value, so we're going to set it back to 0
+		// since at this point, the number of frames is 0!
+		SoundFileData->SetDescription(Description);
+
+		// Set the channel map on the virtual file
+		SoundFileData->SetChannelMap(InChannelMap);
+		int32 Result = SoundFileCommand(FileHandle, SET_CHANNEL_MAP_INFO, (int32*)InChannelMap.GetData(), sizeof(int32)*Description.NumChannels);
+		if (Result)
+		{
+			FString StrErr = SoundFileStrError(nullptr);
+			UE_LOG(LogUnrealAudio, Error, TEXT("Failed to set the channel map on empty file for writing: %s"), *StrErr);
+			return SetError(ESoundFileError::INVALID_CHANNEL_MAP);
+		}
+
+		return ESoundFileError::NONE;
+	}
+
+	ESoundFileError::Type FSoundFile::OpenEmptyFileForExport(const FString& FilePath, const FSoundFileDescription& InDescription, const TArray<ESoundFileChannelMap::Type>& InChannelMap)
+	{
+		// First check the input format to make sure it's valid
+		if (!SoundFileFormatCheck(&InDescription))
+		{
+			UE_LOG(LogUnrealAudio,
+				   Error,
+				   TEXT("Sound file input format (%s - %s) is invalid."),
+				   ESoundFileFormat::ToStringMajor(InDescription.FormatFlags),
+				   ESoundFileFormat::ToStringMinor(InDescription.FormatFlags));
+
+			return SetError(ESoundFileError::INVALID_INPUT_FORMAT);
+		}
+
+		// Make sure we have the right number of channels and our channel map size
+		if (InChannelMap.Num() != InDescription.NumChannels)
+		{
+			UE_LOG(LogUnrealAudio, Error, TEXT("Channel map didn't match the input NumChannels"));
+			return SetError(ESoundFileError::INVALID_CHANNEL_MAP);
+		}
+
+		// Open up the file
+		FSoundFileDescription Description = InDescription;
+		FileHandle = SoundFileOpen(TCHAR_TO_ANSI(*FilePath), ESoundFileOpenMode::WRITING, &Description);
+		if (!FileHandle)
+		{
+			FString StrErr = SoundFileStrError(nullptr);
+			UE_LOG(LogUnrealAudio, Error, TEXT("Failed to open sound file %s: %s"), *FilePath, *StrErr);
+			return SetError(ESoundFileError::FAILED_TO_OPEN);
+		}
+
+		return ESoundFileError::NONE;
+	}
+
+
+	ESoundFileError::Type FSoundFile::ReleaseSoundFileHandle()
+	{
+		if (FileHandle)
+		{
+			SoundFileClose(FileHandle);
+			FileHandle = nullptr;
+		}
+		return ESoundFileError::NONE;
+	}
+
+	ESoundFileError::Type FSoundFile::SetImportFileInfo(const FSoundFileDescription& FileDescription, const FString& FullPath)
+	{
+		if (!SoundFileData.IsValid())
+		{
+			return SetError(ESoundFileError::INVALID_DATA);
+		}
+
+		SoundFileData->SetOriginalDescription(FileDescription);
+		SoundFileData->SetFilePath(FullPath);
+		return ESoundFileError::NONE;
+	}
+
+	ESoundFileError::Type FSoundFile::SetEncodingQuality(double EncodingQuality)
+	{
+		if (FileHandle == nullptr || State != ESoundFileState::UNINITIALIZED)
+		{
+			return SetError(ESoundFileError::INVALID_SOUND_FILE);
+		}
+
+		if (!SoundFileData.IsValid())
+		{
+			return SetError(ESoundFileError::INVALID_DATA);
+		}
+
+		// This only does anything if the format is OGG
+		if ((SoundFileData->GetDescription().FormatFlags & ESoundFileFormat::MAJOR_FORMAT_MASK) == ESoundFileFormat::OGG)
+		{
+
+			int32 Result = SoundFileCommand(FileHandle, SET_ENCODING_QUALITY, &EncodingQuality, sizeof(double));
+			if (Result != 1)
+			{
+				FString StrErr = SoundFileStrError(FileHandle);
+				UE_LOG(LogUnrealAudio, Error, TEXT("Failed to set encoding quality: %s"), *StrErr);
+				return CurrentError = ESoundFileError::BAD_ENCODING_QUALITY;
+			}
+		}
+		return ESoundFileError::NONE;
+	}
+
+	ESoundFileError::Type FSoundFile::BeginImport()
+	{
+		State = ESoundFileState::IMPORTING;
+		return ESoundFileError::NONE;
+	}
+
+	ESoundFileError::Type FSoundFile::EndImport()
+	{
+		// Set the current NumFrames to the original NumFrames
+		((FSoundFileData*)SoundFileData.Get())->SetNumFrames(SoundFileData->GetOriginalDescription().NumFrames);
+		State = ESoundFileState::IMPORTED;
+		return ESoundFileError::NONE;
+	}
+
+	ESoundFileError::Type FSoundFile::SeekFrames(SoundFileCount Offset, ESoundFileSeekMode::Type SeekMode, SoundFileCount& OutOffset)
+	{
+		SoundFileCount Pos = SoundFileSeek(FileHandle, Offset, (int32)SeekMode);
+		if (Pos == -1)
+		{
+			FString StrErr = SoundFileStrError(FileHandle);
+			UE_LOG(LogUnrealAudio, Error, TEXT("Failed to seek file: %s"), *StrErr);
+			return SetError(ESoundFileError::FAILED_TO_SEEK);
+		}
+		return ESoundFileError::NONE;
+	}
+
+	ESoundFileError::Type FSoundFile::ReadFrames(float* DataPtr, SoundFileCount NumFrames, SoundFileCount& NumFramesRead)
+	{
+		NumFramesRead = SoundFileReadFramesFloat(FileHandle, DataPtr, NumFrames);
+		return ESoundFileError::NONE;
+	}
+
+	ESoundFileError::Type FSoundFile::ReadFrames(double* DataPtr, SoundFileCount NumFrames, SoundFileCount& NumFramesRead)
+	{
+		NumFramesRead = SoundFileReadFramesDouble(FileHandle, DataPtr, NumFrames);
+		return ESoundFileError::NONE;
+	}
+
+	ESoundFileError::Type FSoundFile::ReadSamples(float* DataPtr, SoundFileCount NumSamples, SoundFileCount& OutNumSamplesRead)
+	{
+		OutNumSamplesRead = SoundFileReadSamplesFloat(FileHandle, DataPtr, NumSamples);
+		return ESoundFileError::NONE;
+	}
+
+	ESoundFileError::Type FSoundFile::ReadSamples(double* DataPtr, SoundFileCount NumSamples, SoundFileCount& OutNumSamplesRead)
+	{
+		OutNumSamplesRead = SoundFileReadSamplesDouble(FileHandle, DataPtr, NumSamples);
+		return ESoundFileError::NONE;
+	}
+
+	ESoundFileError::Type FSoundFile::WriteFrames(const float* DataPtr, SoundFileCount NumFrames, SoundFileCount& OutNumFramesWritten)
+	{
+		OutNumFramesWritten = SoundFileWriteFramesFloat(FileHandle, DataPtr, NumFrames);
+		return ESoundFileError::NONE;
+	}
+
+	ESoundFileError::Type FSoundFile::WriteFrames(const double* DataPtr, SoundFileCount NumFrames, SoundFileCount& OutNumFramesWritten)
+	{
+		OutNumFramesWritten = SoundFileWriteFramesDouble(FileHandle, DataPtr, NumFrames);
+		return ESoundFileError::NONE;
+	}
+
+	ESoundFileError::Type FSoundFile::WriteSamples(const float* DataPtr, SoundFileCount NumSamples, SoundFileCount& OutNumSampleWritten)
+	{
+		OutNumSampleWritten = SoundFileWriteSamplesFloat(FileHandle, DataPtr, NumSamples);
+		return ESoundFileError::NONE;
+	}
+
+	ESoundFileError::Type FSoundFile::WriteSamples(const double* DataPtr, SoundFileCount NumSamples, SoundFileCount& OutNumSampleWritten)
+	{
+		OutNumSampleWritten = SoundFileWriteSamplesDouble(FileHandle, DataPtr, NumSamples);
+		return ESoundFileError::NONE;
+	}
+
+	ESoundFileError::Type FSoundFile::GetLengthBytes(SoundFileCount& OutLength) const
+	{
+		if (!SoundFileData.IsValid())
+		{
+			return ESoundFileError::INVALID_DATA;
+		}
+
+		OutLength = (SoundFileCount)SoundFileData->GetDataSize();
+		return ESoundFileError::NONE;
+	}
+
+	ESoundFileError::Type FSoundFile::SeekBytes(SoundFileCount Offset, ESoundFileSeekMode::Type SeekMode, SoundFileCount& OutOffset)
+	{
+		if (!SoundFileData.IsValid())
+		{
+			return ESoundFileError::INVALID_DATA;
+		}
+
+		// If there's no data, then seek does nothing
+		SoundFileCount MaxBytes = SoundFileData->GetDataSize();
+		if (MaxBytes == 0)
+		{
+			OutOffset = 0;
+			CurrentIndexBytes = 0;
+			return ESoundFileError::NONE;
+		}
+
+		switch (SeekMode)
+		{
+			case ESoundFileSeekMode::FROM_START:
+			CurrentIndexBytes = Offset;
+			break;
+
+			case ESoundFileSeekMode::FROM_CURRENT:
+			CurrentIndexBytes += Offset;
+			break;
+
+			case ESoundFileSeekMode::FROM_END:
+			CurrentIndexBytes = MaxBytes + Offset;
+			break;
+
+			default:
+			checkf(false, TEXT("Uknown seek mode!"));
+			break;
+		}
+
+		// Wrap the byte index to fall between 0 and MaxBytes
+		while (CurrentIndexBytes < 0)
+		{
+			CurrentIndexBytes += MaxBytes;
+		}
+
+		while (CurrentIndexBytes >= MaxBytes)
+		{
+			CurrentIndexBytes -= MaxBytes;
+		}
+
+		OutOffset = CurrentIndexBytes;
+		return ESoundFileError::NONE;
+	}
+
+	ESoundFileError::Type FSoundFile::ReadBytes(void* DataPtr, SoundFileCount NumBytes, SoundFileCount& NumBytesRead)
+	{
+		if (!SoundFileData.IsValid())
+		{
+			return ESoundFileError::INVALID_DATA;
+		}
+
+		SoundFileCount EndByte = CurrentIndexBytes + NumBytes;
+		SoundFileCount MaxBytes = (SoundFileCount)SoundFileData->GetDataSize();
+		if (EndByte > MaxBytes)
+		{
+			NumBytes = MaxBytes - CurrentIndexBytes;
+		}
+
+		if (NumBytes > 0)
+		{
+			const uint8* DataArrayData = SoundFileData->GetData();
+			FMemory::Memcpy(DataPtr, (const void*)&DataArrayData[CurrentIndexBytes], NumBytes);
+ 			CurrentIndexBytes += NumBytes;
+		}
+		NumBytesRead = NumBytes;
+		return ESoundFileError::NONE;
+	}
+
+	ESoundFileError::Type FSoundFile::WriteBytes(const void* DataPtr, SoundFileCount NumBytes, SoundFileCount& NumBytesWritten)
+	{
+		const uint8* InDataBytes = (const uint8*)DataPtr;
+		FSoundFileData* Data = (FSoundFileData*)SoundFileData.Get();
+		SoundFileCount MaxBytes = Data->GetDataArray().Num();
+
+		// Append the data to the buffer
+		if (CurrentIndexBytes == MaxBytes)
+		{
+			Data->GetDataArray().Append(InDataBytes, NumBytes);
+			CurrentIndexBytes += NumBytes;
+		}
+		else if ((CurrentIndexBytes + NumBytes) < MaxBytes)
+		{
+			// Write over the existing data in the buffer
+			for (SoundFileCount i = 0; i < NumBytes; ++i)
+			{
+				Data->GetDataArray()[CurrentIndexBytes++] = InDataBytes[i];
+			}
+		}
+		else
+		{
+			// Overwrite some data until end of current buffer size
+			SoundFileCount RemainingBytes = MaxBytes - CurrentIndexBytes;
+			SoundFileCount InputBufferIndex = 0;
+			for (InputBufferIndex = 0; InputBufferIndex < RemainingBytes; ++InputBufferIndex)
+			{
+				Data->GetDataArray()[CurrentIndexBytes++] = InDataBytes[InputBufferIndex];
+			}
+
+			// Now append the remainder of the input buffer to the internal data buffer
+			const uint8* RemainingData = &InDataBytes[InputBufferIndex];
+			RemainingBytes = NumBytes - RemainingBytes;
+			Data->GetDataArray().Append(RemainingData, RemainingBytes);
+			CurrentIndexBytes += RemainingBytes;
+			check(CurrentIndexBytes == Data->GetDataArray().Num());
+		}
+		NumBytesWritten = NumBytes;
+		return ESoundFileError::NONE;
+	}
+
+	SoundFileCount FSoundFile::GetOffsetBytes() const
+	{
+		return CurrentIndexBytes;
+	}
+
+	FSoundFileData::FSoundFileData(ESoundFileOpenMode::Type InMode)
+		: DataBuffer(nullptr)
+		, DataBufferSize(0)
+		, Mode(InMode)
+	{
+		Description = { 0 };
+		OriginalDescription = { 0 };
+	}
+
+	FSoundFileData::~FSoundFileData()
+	{
+	}
+
+	const FString& FSoundFileData::GetFilePath() const
+	{
+		return FilePath;
+	}
+
+	void FSoundFileData::SetFilePath(const FString& InPath)
+	{
+		FilePath = InPath;
+	}
+
+	const FSoundFileDescription& FSoundFileData::GetDescription() const
+	{
+		return Description;
+	}
+
+	void FSoundFileData::SetDescription(const FSoundFileDescription& InDescription)
+	{
+		Description = InDescription;
+	}
+
+	const FSoundFileDescription& FSoundFileData::GetOriginalDescription() const
+	{
+		return OriginalDescription;
+	}
+
+	void FSoundFileData::SetOriginalDescription(const FSoundFileDescription& InDescription)
+	{
+		OriginalDescription = InDescription;
+	}
+
+	const TArray<ESoundFileChannelMap::Type>& FSoundFileData::GetChannelMap() const
+	{
+		return ChannelMap;
+	}
+
+	void FSoundFileData::SetChannelMap(const TArray<ESoundFileChannelMap::Type>& InChannelMap)
+	{
+		ChannelMap = InChannelMap;
+	}
+
+	void FSoundFileData::SetNumFrames(int64 InNumFrames)
+	{
+		Description.NumFrames = InNumFrames;
+	}
+
+	const uint8* FSoundFileData::GetData() const
+	{
+		if (Mode == ESoundFileOpenMode::WRITING)
+		{
+			return DataArray.GetData();
+		}
+		else
+		{
+			return DataBuffer;
+		}
+	}
+
+	int64 FSoundFileData::GetDataSize() const
+	{
+		if (Mode == ESoundFileOpenMode::WRITING)
+		{
+			return (int64)DataArray.Num();
+		}
+		else
+		{
+			return DataBufferSize;
+		}
+	}
+
+	TArray<uint8>& FSoundFileData::GetDataArray()
+	{
+		// If we're accessing the array, make sure this sound file data object is in write-mode
+		check(Mode == ESoundFileOpenMode::WRITING);
+		return DataArray;
+	}
+
+	void FSoundFileData::SetData(uint8* InDataPtr, int64 InDataSize)
+	{
+		// Make sure this sound file data object wasn't opened for writing
+		check(Mode == ESoundFileOpenMode::READING);
+		DataBuffer = InDataPtr;
+		DataBufferSize = InDataSize;
+	}
+
+	// Exported Functions
+
+	TSharedPtr<ISoundFileData> CreateSoundFileData()
+	{
+		// Only internal sound file code creates write-able sound file data
+		return TSharedPtr<ISoundFileData>(new FSoundFileData(ESoundFileOpenMode::READING));
+	}
+
+	TSharedPtr<ISoundFile> CreateSoundFile()
+	{
+		return TSharedPtr<ISoundFile>(new FSoundFile());
+	}
+
+	// UnrealAudio Module Functions
+
+	static void* GetSoundFileDllHandle()
+	{
+		void* DllHandle = nullptr;
+#if PLATFORM_WINDOWS
+		FString Path = FPaths::EngineDir() / FString(TEXT("Binaries/ThirdParty/libsndfile/Win64/"));
+		FPlatformProcess::PushDllDirectory(*Path);
+		DllHandle = FPlatformProcess::GetDllHandle(*(Path + "libsndfile-1.dll"));
+		FPlatformProcess::PopDllDirectory(*Path);
+#elif PLATFORM_MAC
+//		FString Path = FPaths::EngineDir() / FString(TEXT("Binaries/ThirdParty/libsndfile/Mac/"));
+//		FPlatformProcess::PushDllDirectory(*Path);
+		DllHandle = FPlatformProcess::GetDllHandle(TEXT("libsndfile.1.dylib"));
+//		FPlatformProcess::PopDllDirectory(*Path);
+#endif
+		return DllHandle;
+	}
+
+	bool FUnrealAudioModule::LoadSoundFileLib()
+	{
+		SoundFileDllHandle = GetSoundFileDllHandle();
+		if (!SoundFileDllHandle)
+		{
+			UE_LOG(LogUnrealAudio, Error, TEXT("Failed to load Sound File dll"));
+			return false;
+		}
+
+		bool bSuccess = true;
+
+		void* LambdaDLLHandle = SoundFileDllHandle;
+
+		// Helper function to load DLL exports and report warnings
+		auto GetSoundFileDllExport = [&LambdaDLLHandle](const TCHAR* FuncName, bool& bInSuccess) -> void*
+		{
+			if (bInSuccess)
+			{
+				void* FuncPtr = FPlatformProcess::GetDllExport(LambdaDLLHandle, FuncName);
+				if (FuncPtr == nullptr)
+				{
+					bInSuccess = false;
+					UE_LOG(LogUnrealAudio, Warning, TEXT("Failed to locate the expected DLL import function '%s' in the SoundFile DLL."), *FuncName);
+					FPlatformProcess::FreeDllHandle(LambdaDLLHandle);
+					LambdaDLLHandle = nullptr;
+				}
+				return FuncPtr;
+			}
+			else
+			{
+				return nullptr;
+			}
+		};
+
+		SoundFileOpen = (SoundFileOpenFuncPtr)GetSoundFileDllExport(TEXT("sf_open"), bSuccess);
+		SoundFileOpenVirtual = (SoundFileOpenVirtualFuncPtr)GetSoundFileDllExport(TEXT("sf_open_virtual"), bSuccess);
+		SoundFileClose = (SoundFileCloseFuncPtr)GetSoundFileDllExport(TEXT("sf_close"), bSuccess);
+		SoundFileError = (SoundFileErrorFuncPtr)GetSoundFileDllExport(TEXT("sf_error"), bSuccess);
+		SoundFileStrError = (SoundFileStrErrorFuncPtr)GetSoundFileDllExport(TEXT("sf_strerror"), bSuccess);
+		SoundFileErrorNumber = (SoundFileErrorNumberFuncPtr)GetSoundFileDllExport(TEXT("sf_error_number"), bSuccess);
+		SoundFileCommand = (SoundFileCommandFuncPtr)GetSoundFileDllExport(TEXT("sf_command"), bSuccess);
+		SoundFileFormatCheck = (SoundFileFormatCheckFuncPtr)GetSoundFileDllExport(TEXT("sf_format_check"), bSuccess);
+		SoundFileSeek = (SoundFileSeekFuncPtr)GetSoundFileDllExport(TEXT("sf_seek"), bSuccess);
+		SoundFileGetVersion = (SoundFileGetVersionFuncPtr)GetSoundFileDllExport(TEXT("sf_version_string"), bSuccess);
+		SoundFileReadFramesFloat = (SoundFileReadFramesFloatFuncPtr)GetSoundFileDllExport(TEXT("sf_readf_float"), bSuccess);
+		SoundFileReadFramesDouble = (SoundFileReadFramesDoubleFuncPtr)GetSoundFileDllExport(TEXT("sf_readf_double"), bSuccess);
+		SoundFileWriteFramesFloat = (SoundFileWriteFramesFloatFuncPtr)GetSoundFileDllExport(TEXT("sf_writef_float"), bSuccess);
+		SoundFileWriteFramesDouble = (SoundFileWriteFramesDoubleFuncPtr)GetSoundFileDllExport(TEXT("sf_writef_double"), bSuccess);
+		SoundFileReadSamplesFloat = (SoundFileReadSamplesFloatFuncPtr)GetSoundFileDllExport(TEXT("sf_read_float"), bSuccess);
+		SoundFileReadSamplesDouble = (SoundFileReadSamplesDoubleFuncPtr)GetSoundFileDllExport(TEXT("sf_read_double"), bSuccess);
+		SoundFileWriteSamplesFloat = (SoundFileWriteSamplesFloatFuncPtr)GetSoundFileDllExport(TEXT("sf_write_float"), bSuccess);
+		SoundFileWriteSamplesDouble = (SoundFileWriteSamplesDoubleFuncPtr)GetSoundFileDllExport(TEXT("sf_write_double"), bSuccess);
+
+		// make sure we're successful
+		check(bSuccess);
+		return bSuccess;
+	}
+
+	bool FUnrealAudioModule::ShutdownSoundFileLib()
+	{
+		if (SoundFileDllHandle)
+		{
+			FPlatformProcess::FreeDllHandle(SoundFileDllHandle);
+			SoundFileDllHandle = nullptr;
+			SoundFileOpen = nullptr;
+			SoundFileOpenVirtual = nullptr;
+			SoundFileClose = nullptr;
+			SoundFileError = nullptr;
+			SoundFileStrError = nullptr;
+			SoundFileErrorNumber = nullptr;
+			SoundFileCommand = nullptr;
+			SoundFileFormatCheck = nullptr;
+			SoundFileSeek = nullptr;
+			SoundFileGetVersion = nullptr;
+			SoundFileReadFramesFloat = nullptr;
+			SoundFileReadFramesDouble = nullptr;
+			SoundFileWriteFramesFloat = nullptr;
+			SoundFileWriteFramesDouble = nullptr;
+			SoundFileReadSamplesFloat = nullptr;
+			SoundFileReadSamplesDouble = nullptr;
+			SoundFileWriteSamplesFloat = nullptr;
+			SoundFileWriteSamplesDouble = nullptr;
+		}
+		return true;
+	}
+
+}
+
+#endif // #if ENABLE_UNREAL_AUDIO
\ No newline at end of file
diff --git a/Engine/Source/Runtime/UnrealAudio/Private/UnrealAudioSoundFileExport.cpp b/Engine/Source/Runtime/UnrealAudio/Private/UnrealAudioSoundFileExport.cpp
new file mode 100644
--- /dev/null
+++ b/Engine/Source/Runtime/UnrealAudio/Private/UnrealAudioSoundFileExport.cpp
@@ -0,0 +1,110 @@
+#include "UnrealAudioPrivate.h"
+#include "UnrealAudioSoundFile.h"
+#include "UnrealAudioModule.h"
+#include "ModuleManager.h"
+#include "UnrealAudioSoundFile.h"
+#include "UnrealAudioSoundFileInternal.h"
+
+#if ENABLE_UNREAL_AUDIO
+
+#define SOUND_EXPORT_CHECK(Result)			\
+	if (Result != ESoundFileError::NONE)	\
+	{										\
+		return;								\
+	}
+
+namespace UAudio
+{
+	/**
+	FAsyncSoundFileExportTask
+	Task which processes a sound file export request on a background thread.
+	*/
+	class FAsyncSoundFileExportTask
+	{
+	public:
+		FAsyncSoundFileExportTask(FUnrealAudioModule* InAudioModule, TSharedPtr<ISoundFileData> InSoundFileData, const FString* InExportPath);
+		~FAsyncSoundFileExportTask();
+
+		void DoWork();
+		bool CanAbandon() { return true; }
+		void Abandon() {}
+
+		FORCEINLINE TStatId GetStatId() const
+		{
+			RETURN_QUICK_DECLARE_CYCLE_STAT(FAsyncSoundFileExportTask, STATGROUP_ThreadPoolAsyncTasks);
+		}
+
+		FUnrealAudioModule* AudioModule;
+		TSharedPtr<ISoundFileData> SoundFileData;
+		FString ExportPath;
+		friend class FAsyncTask<FAsyncSoundFileExportTask>;
+	};
+
+	FAsyncSoundFileExportTask::FAsyncSoundFileExportTask(FUnrealAudioModule* InAudioModule, TSharedPtr<ISoundFileData> InSoundFileData, const FString* InExportPath)
+		: AudioModule(InAudioModule)
+		, SoundFileData(InSoundFileData)
+		, ExportPath(*InExportPath)
+	{
+		// make sure the audio module knows we're doing background task work so it doesn't shutdown while we're working here
+		AudioModule->IncrementBackgroundTaskCount();
+	}
+
+	FAsyncSoundFileExportTask::~FAsyncSoundFileExportTask()
+	{
+		// Let the audio module know we're done doing this background task
+		AudioModule->DecrementBackgroundTaskCount();
+	}
+
+	void FAsyncSoundFileExportTask::DoWork()
+	{
+		// Create a new sound file object which we will pass our data to so we can read from it
+		FSoundFile SoundFileInput;
+
+		ESoundFileError::Type Error = SoundFileInput.Initialize(SoundFileData);
+		SOUND_EXPORT_CHECK(Error);
+
+		// Create a new sound file object which will be written to disk
+		FSoundFile SoundFileOutput;
+
+		FSoundFileDescription OriginalDescription = SoundFileData->GetOriginalDescription();
+		const TArray<ESoundFileChannelMap::Type>& ChannelMap = SoundFileData->GetChannelMap();
+
+		Error = SoundFileOutput.OpenEmptyFileForExport(ExportPath, OriginalDescription, ChannelMap);
+		SOUND_EXPORT_CHECK(Error);
+
+		// Create a buffer to do the processing 
+		SoundFileCount ProcessBufferSamples = 1024 * OriginalDescription.NumChannels;
+		TArray<double> ProcessBuffer;
+		ProcessBuffer.Init(0.0, ProcessBufferSamples);
+
+		// Now perform the encoding to the target file
+		SoundFileCount SamplesRead = 0;
+		Error = SoundFileInput.ReadSamples(ProcessBuffer.GetData(), ProcessBufferSamples, SamplesRead);
+		SOUND_EXPORT_CHECK(Error);
+		
+		while (SamplesRead)
+		{
+			SOUND_EXPORT_CHECK(SoundFileInput.GetError());
+			SoundFileCount SamplesWritten;
+
+			Error = SoundFileOutput.WriteSamples(ProcessBuffer.GetData(), SamplesRead, SamplesWritten);
+			SOUND_EXPORT_CHECK(Error);
+
+			Error = SoundFileInput.ReadSamples(ProcessBuffer.GetData(), ProcessBufferSamples, SamplesRead);
+			SOUND_EXPORT_CHECK(Error);
+		}
+
+		SoundFileInput.ReleaseSoundFileHandle();
+		SoundFileOutput.ReleaseSoundFileHandle();
+	}
+	
+	void FUnrealAudioModule::ExportSound(TSharedPtr<ISoundFileData> SoundFileData, const FString& ExportPath)
+	{
+		// Make an auto-deleting task to perform the import work in the background task thread pool
+		FAutoDeleteAsyncTask<FAsyncSoundFileExportTask>* Task = new FAutoDeleteAsyncTask<FAsyncSoundFileExportTask>(this, SoundFileData, &ExportPath);
+		Task->StartBackgroundTask();
+	}
+
+}
+
+#endif // #if ENABLE_UNREAL_AUDIO
diff --git a/Engine/Source/Runtime/UnrealAudio/Private/UnrealAudioSoundFileImport.cpp b/Engine/Source/Runtime/UnrealAudio/Private/UnrealAudioSoundFileImport.cpp
new file mode 100644
--- /dev/null
+++ b/Engine/Source/Runtime/UnrealAudio/Private/UnrealAudioSoundFileImport.cpp
@@ -0,0 +1,193 @@
+#include "UnrealAudioPrivate.h"
+#include "UnrealAudioSoundFile.h"
+#include "UnrealAudioModule.h"
+#include "ModuleManager.h"
+#include "UnrealAudioSoundFile.h"
+#include "UnrealAudioSoundFileInternal.h"
+
+#if ENABLE_UNREAL_AUDIO
+
+
+#define SOUND_IMPORT_CHECK(Result)							\
+	if (Result != ESoundFileError::NONE)					\
+	{														\
+		((ISoundFileInternal*)SoundFile.Get())->SetError(Result);	\
+		return;												\
+	}
+
+
+namespace UAudio
+{
+	/**
+	FAsyncSoundFileImportTask
+	Task which processes a sound file import request on a background thread.
+	*/
+	class FAsyncSoundFileImportTask
+	{
+	public:
+		FAsyncSoundFileImportTask(FUnrealAudioModule* InAudioModule, FSoundFileImportSettings InImportSettings, TSharedPtr<ISoundFile> InSoundFile);
+		~FAsyncSoundFileImportTask();
+
+		void DoWork();
+		bool CanAbandon() { return true; }
+		void Abandon() {}
+
+		FORCEINLINE TStatId GetStatId() const
+		{
+			RETURN_QUICK_DECLARE_CYCLE_STAT(FAsyncSoundFileImportTask, STATGROUP_ThreadPoolAsyncTasks);
+		}
+
+		FUnrealAudioModule* AudioModule;
+		FSoundFileImportSettings ImportSettings;
+		TSharedPtr<ISoundFile> SoundFile;
+		friend class FAsyncTask<FAsyncSoundFileImportTask>;
+	};
+
+	FAsyncSoundFileImportTask::FAsyncSoundFileImportTask(FUnrealAudioModule* InAudioModule, FSoundFileImportSettings InImportSettings, TSharedPtr<ISoundFile> InSoundFile)
+		: AudioModule(InAudioModule)
+		, ImportSettings(InImportSettings)
+		, SoundFile(InSoundFile)
+	{
+		// make sure the audio module knows we're doing background task work so it doesn't shutdown while we're working here
+		AudioModule->IncrementBackgroundTaskCount();
+	}
+
+	FAsyncSoundFileImportTask::~FAsyncSoundFileImportTask()
+	{
+		// Let the audio module know we're done doing this background task
+		AudioModule->DecrementBackgroundTaskCount();
+	}
+
+	void FAsyncSoundFileImportTask::DoWork()
+	{
+
+		// Create a new sound file object
+		TScopedPointer<FSoundFile> SoundFileInput = TScopedPointer<FSoundFile>(new FSoundFile());
+		
+		// Open the file
+		ESoundFileError::Type Error = SoundFileInput->OpenFileForReading(ImportSettings.SoundFilePath);
+		SOUND_IMPORT_CHECK(Error);
+
+		TSharedPtr<ISoundFileData> SoundFileData;
+		Error = SoundFileInput->GetSoundFileData(SoundFileData);
+		SOUND_IMPORT_CHECK(Error);
+
+		// Get the input file's description
+		const FSoundFileDescription& InputDescription = SoundFileData->GetDescription();
+
+		// Get the input file's channel map
+		const TArray<ESoundFileChannelMap::Type>& InputChannelMap = SoundFileData->GetChannelMap();
+
+		// Build a description for the new file
+		FSoundFileDescription NewSoundFileDescription;
+		NewSoundFileDescription.NumChannels = InputDescription.NumChannels;
+		NewSoundFileDescription.NumFrames = InputDescription.NumFrames;
+		NewSoundFileDescription.FormatFlags = ImportSettings.Format;
+		NewSoundFileDescription.SampleRate = ImportSettings.SampleRate;
+		NewSoundFileDescription.NumSections = 0;
+		NewSoundFileDescription.bIsSeekable = 1;
+
+		// Open it as an empty file for reading and writing
+		ISoundFileInternal* SoundFileInternal = (ISoundFileInternal*)SoundFile.Get();
+		Error = SoundFileInternal->OpenEmptyFileForImport(NewSoundFileDescription, InputChannelMap);
+		SOUND_IMPORT_CHECK(Error);
+
+		// Set the original description on the new sound file
+		Error = SoundFileInternal->SetImportFileInfo(InputDescription, ImportSettings.SoundFilePath);
+		SOUND_IMPORT_CHECK(Error);
+
+		// Set the encoding quality (will only do anything if import target is Ogg-Vorbis)
+		Error = SoundFileInternal->SetEncodingQuality(ImportSettings.EncodingQuality);
+		SOUND_IMPORT_CHECK(Error);
+
+		// Set the state of the sound file to be loading
+		Error = SoundFileInternal->BeginImport();
+		SOUND_IMPORT_CHECK(Error);
+
+		// Create a buffer to do the processing 
+		SoundFileCount ProcessBufferSamples = 1024 * NewSoundFileDescription.NumChannels;
+		TArray<double> ProcessBuffer;
+		ProcessBuffer.Init(0.0, ProcessBufferSamples);
+
+		// Find the max value if we've been told to do peak normalization on import
+		double MaxValue = 0.0;
+		SoundFileCount SamplesRead = 0;
+		bool bPerformPeakNormalization = ImportSettings.bPerformPeakNormalization;
+		if (bPerformPeakNormalization)
+		{
+			Error = SoundFileInput->ReadSamples(ProcessBuffer.GetData(), ProcessBufferSamples, SamplesRead);
+			SOUND_IMPORT_CHECK(Error);
+
+			while (SamplesRead)
+			{
+				for (SoundFileCount Sample = 0; Sample < SamplesRead; ++Sample)
+				{
+					if (ProcessBuffer[Sample] > FMath::Abs(MaxValue))
+					{
+						MaxValue = ProcessBuffer[Sample];
+					}
+				}
+
+				Error = SoundFileInput->ReadSamples(ProcessBuffer.GetData(), ProcessBufferSamples, SamplesRead);
+				SOUND_IMPORT_CHECK(Error);
+			}
+
+			// If this happens, it means we have a totally silent file
+			if (MaxValue == 0.0)
+			{
+				bPerformPeakNormalization = false;
+			}
+
+			// Seek the file back to the beginning
+			SoundFileCount OutOffset;
+			SoundFileInput->SeekFrames(0, ESoundFileSeekMode::FROM_START, OutOffset);
+		}
+
+		// Now perform the encoding to the target file
+		Error = SoundFileInput->ReadSamples(ProcessBuffer.GetData(), ProcessBufferSamples, SamplesRead);
+		SOUND_IMPORT_CHECK(Error);
+
+		while (SamplesRead)
+		{
+			SOUND_IMPORT_CHECK(SoundFileInput->GetError());
+
+			// Normalize the samples if we're told to
+			if (bPerformPeakNormalization)
+			{
+				for (int32 Sample = 0; Sample < SamplesRead; ++Sample)
+				{
+					ProcessBuffer[Sample] /= MaxValue;
+				}
+			}
+
+			SoundFileCount SamplesWritten;
+			Error = SoundFileInternal->WriteSamples((const double*)ProcessBuffer.GetData(), SamplesRead, SamplesWritten);
+			SOUND_IMPORT_CHECK(Error);
+
+			Error = SoundFileInput->ReadSamples(ProcessBuffer.GetData(), ProcessBufferSamples, SamplesRead);
+			SOUND_IMPORT_CHECK(Error);
+		}
+
+		SoundFileInput->ReleaseSoundFileHandle();
+		SoundFileInternal->ReleaseSoundFileHandle();
+
+		// We're done doing the encoding
+		Error = SoundFileInternal->EndImport();
+		SOUND_IMPORT_CHECK(Error);
+	}
+
+	TSharedPtr<ISoundFile> FUnrealAudioModule::ImportSound(const FSoundFileImportSettings& ImportSettings)
+	{
+		TSharedPtr<ISoundFile> ImportedSoundFile = TSharedPtr<ISoundFile>(new FSoundFile());
+
+		// Make an auto-deleting task to perform the import work in the background task thread pool
+		FAutoDeleteAsyncTask<FAsyncSoundFileImportTask>* Task = new FAutoDeleteAsyncTask<FAsyncSoundFileImportTask>(this, ImportSettings, ImportedSoundFile);
+		Task->StartBackgroundTask();
+
+		return ImportedSoundFile;
+	}
+
+}
+
+#endif // #if ENABLE_UNREAL_AUDIO
+
diff --git a/Engine/Source/Runtime/UnrealAudio/Private/UnrealAudioSoundFileInternal.h b/Engine/Source/Runtime/UnrealAudio/Private/UnrealAudioSoundFileInternal.h
new file mode 100644
--- /dev/null
+++ b/Engine/Source/Runtime/UnrealAudio/Private/UnrealAudioSoundFileInternal.h
@@ -0,0 +1,328 @@
+// Copyright 1998-2015 Epic Games, Inc. All Rights Reserved.
+
+#pragma once 
+
+#include "UnrealAudioModule.h"
+#include "UnrealAudioTypes.h"
+#include "UnrealAudioSoundFile.h"
+
+#if ENABLE_UNREAL_AUDIO
+
+namespace UAudio
+{
+	typedef int64 SoundFileCount;
+	typedef struct SoundFileHandleOpaque SoundFileHandle;
+
+	/*
+	* ESoundFileSeekMode
+	* Specifies the how to seek a sound file
+	**/
+	namespace ESoundFileSeekMode
+	{
+		enum Type
+		{
+			// Seek the file from the start of the file.
+			FROM_START = 0,
+
+			// Seek the file from the current read/write position.
+			FROM_CURRENT = 1,
+
+			// Seek the file from the end position.
+			FROM_END = 2,
+		};
+	}
+
+	/*
+	* ESoundFileOpenMode
+	* Specifies the how to open a sound file1
+	**/
+	namespace ESoundFileOpenMode
+	{
+		enum Type
+		{
+			// The sound file is open for reading.
+			READING = 0x10,
+
+			// The sound file is open for writing.
+			WRITING = 0x20,
+
+			UNKNOWN = 0,
+		};
+	}
+
+
+	/**
+	ISoundFileVirtual
+	Interface for specifying the virtual sound file i/o operations.
+	*/
+	class ISoundFileVirtual : public ISoundFile
+	{
+	public:
+		virtual ~ISoundFileVirtual() {}
+
+		/**
+		* Gets the length of the file in bytes.
+		* @param OutLength The returned length of the file.
+		*/
+		virtual ESoundFileError::Type GetLengthBytes(SoundFileCount& OutLength) const = 0;
+
+		/**
+		* Seeks the read/write position offset in bytes.
+		* @param Offset The offset in bytes of the file
+		* @param SeekMode The type of seek to use when performing the seek.
+		* @param OutOffset Optionally retrieve the actual offset in bytes after the seek operation.
+		*/
+		virtual ESoundFileError::Type SeekBytes(SoundFileCount Offset, ESoundFileSeekMode::Type SeekMode, SoundFileCount& OutOffset) = 0;
+
+		/**
+		* Reads the number of bytes from the given buffer ptr.
+		* @param DataPtr Pointer to a buffer for copying into.
+		* @param NumBytes The number of bytes to read.
+		* @param NumBytesRead Returns the number of bytes actually read.
+		*/
+		virtual ESoundFileError::Type ReadBytes(void* DataPtr, SoundFileCount NumBytes, SoundFileCount& NumBytesRead) = 0;
+
+		/**
+		* Writes the number of bytes into the given buffer ptr.
+		* @param DataPtr Pointer to an input buffer
+		* @param NumBytes The number of bytes to write.
+		* @param NumBytesWritten The number of bytes actually written.
+		*/
+		virtual ESoundFileError::Type WriteBytes(const void* DataPtr, SoundFileCount NumBytes, SoundFileCount& NumBytesWritten) = 0;
+
+		/** @return Gets the current file offset in bytes. */
+		virtual SoundFileCount GetOffsetBytes() const = 0;
+	};
+
+	/**
+	FSoundFile
+	Internal SoundFile API
+	*/
+	class ISoundFileInternal : public ISoundFileVirtual
+	{
+	public:
+		virtual ~ISoundFileInternal() {}
+
+		/**
+		* Sets the error on the sound file. Internal audio systems can cause audio files to get errors. For example, if a failure occurs
+		* during import due to an error on the input sound file, the output (imported) sound file needs to also get the error.
+		*/
+		virtual ESoundFileError::Type SetError(ESoundFileError::Type Error) = 0;
+
+		/**
+		* Opens a sound file from the given path.
+		* @param FilePath Path to the sound file.
+		* @param Mode The mode to open the sound file in.
+		*/
+		virtual ESoundFileError::Type OpenFileForReading(const FString& FilePath) = 0;
+
+		/**
+		* Opens an empty virtual sound file of the given description for writing to.
+		* @param FileDescription Description of the type of file this is to be.
+		* @param InChannelMap The input channel map to use for the file.
+		* @param Mode The mode to open the sound file in.
+		*/
+		virtual ESoundFileError::Type OpenEmptyFileForImport(const FSoundFileDescription& FileDescription, const TArray<ESoundFileChannelMap::Type>& InChannelMap) = 0;
+
+		/**
+		* Opens an empty sound file on disk of the given description for writing to.
+		* @param FileDescription Description of the type of file this is to be.
+		* @param InChannelMap The input channel map to use for the file.
+		* @param Mode The mode to open the sound file in.
+		*/
+		virtual ESoundFileError::Type OpenEmptyFileForExport(const FString& FilePath, const FSoundFileDescription& FileDescription, const TArray<ESoundFileChannelMap::Type>& InChannelMap) = 0;
+
+		/** Releases the internal sound file handle if its open */
+		virtual ESoundFileError::Type ReleaseSoundFileHandle() = 0;
+
+		/**
+		* Sets the description of the original file used to import this file.
+		* @param FileDescription Description of the original file used to import this file.
+		*/
+		virtual ESoundFileError::Type SetImportFileInfo(const FSoundFileDescription& FileDescription, const FString& FullPath) = 0;
+
+		/**
+		* Sets the encoding quality of a sound file.
+		* @param EncodingQuality The encoding quality of a sound file (between 0.0 (low) and 1.0 (high))
+		* @note This can only be set during sound file import before importing begins, but after the file has been opened.
+		*		It can also only be set on OGG-Vorbis files.
+		*/
+		virtual ESoundFileError::Type SetEncodingQuality(double EncodingQuality) = 0;
+
+		/**
+		* Sets the state of the sound file to begin loading
+		*/
+		virtual ESoundFileError::Type BeginImport() = 0;
+
+		/**
+		* Sets the state of the sound file to end loading
+		*/
+		virtual ESoundFileError::Type EndImport() = 0;
+
+		/**
+		* Seeks the read/write position offset in frames.
+		* @param Offset The offset in frames of the file
+		* @param SeekMode The type of seek to use when performing the seek.
+		* @param OutOffset Optionally retrieve the actual offset in bytes after the seek operation.
+		*/
+		virtual ESoundFileError::Type SeekFrames(SoundFileCount Offset, ESoundFileSeekMode::Type SeekMode, SoundFileCount& OutOffset) = 0;
+
+		/**
+		* Reads out the number of frames of data as interleaved floats/doubles.
+		* @param DataPtr Pointer to a buffer to copy data into.
+		* @param NumFrames The number of frames to read.
+		* @param NumFramesRead Returns the number of frames actually read.
+		*/
+		virtual ESoundFileError::Type ReadFrames(float* DataPtr, SoundFileCount NumFrames, SoundFileCount& NumFramesRead) = 0;
+		virtual ESoundFileError::Type ReadFrames(double* DataPtr, SoundFileCount NumFrames, SoundFileCount& NumFramesRead) = 0;
+
+		/**
+		* Reads out the number of samples of data as floats/doubles.
+		* @param DataPtr Pointer to a buffer for copy data into.
+		* @param NumSamples The number of samples to read.
+		* @param OutNumSamplesRead Returns the number of samples actually read.
+		*/
+		virtual ESoundFileError::Type ReadSamples(float* DataPtr, SoundFileCount NumSamples, SoundFileCount& OutNumSamplesRead) = 0;
+		virtual ESoundFileError::Type ReadSamples(double* DataPtr, SoundFileCount NumSamples, SoundFileCount& OutNumSamplesRead) = 0;
+
+		/**
+		* Writes the number of frames of data as interleaved floats/doubles.
+		* @param DataPtr Pointer to a buffer for copying data from.
+		* @param NumFrames The number of frames to write.
+		* @param OutNumFramesWritten Returns the number of frames actually written.
+		*/
+		virtual ESoundFileError::Type WriteFrames(const float* Data, SoundFileCount NumFrames, SoundFileCount& OutNumFramesWritten) = 0;
+		virtual ESoundFileError::Type WriteFrames(const double* Data, SoundFileCount NumFrames, SoundFileCount& OutNumFramesWritten) = 0;
+
+		/**
+		* Writes the number of samples of data as floats/doubles.
+		* @param DataPtr Pointer to a buffer for copying data from.
+		* @param NumSamples The number of samples to write.
+		* @param OutNumSamplesWritten Returns the number of samples actually written.
+		*/
+		virtual ESoundFileError::Type WriteSamples(const float* Data, SoundFileCount NumSamples, SoundFileCount& OutNumSamplesWritten) = 0;
+		virtual ESoundFileError::Type WriteSamples(const double* Data, SoundFileCount NumSamples, SoundFileCount& OutNumSamplesWritten) = 0;
+	};
+
+
+	/**
+	FSoundFile
+	Concrete implementation of ISoundFile which uses libsndfile for file reading/writing.
+	*/
+	class FSoundFile : public ISoundFileInternal
+	{
+	public:
+		FSoundFile();
+		~FSoundFile();
+
+		// ISoundFile implementation
+		ESoundFileError::Type Initialize(TSharedPtr<ISoundFileData> InSoundFileData) override;
+		ESoundFileError::Type GetName(FString& Name) override;
+		ESoundFileError::Type GetSoundFileData(TSharedPtr<ISoundFileData>& OutSoundFileData) override;
+		ESoundFileState::Type GetState() const override;
+		ESoundFileError::Type GetError() const override;
+
+		// ISoundFileInternal implementation
+		ESoundFileError::Type SetError(ESoundFileError::Type Error) override;
+		ESoundFileError::Type OpenFileForReading(const FString& FilePath) override;
+		ESoundFileError::Type OpenEmptyFileForImport(const FSoundFileDescription& FileDescription, const TArray<ESoundFileChannelMap::Type>& InChannelMap) override;
+		ESoundFileError::Type OpenEmptyFileForExport(const FString& FilePath, const FSoundFileDescription& FileDescription, const TArray<ESoundFileChannelMap::Type>& InChannelMap) override;
+		ESoundFileError::Type ReleaseSoundFileHandle() override;
+
+		ESoundFileError::Type SetImportFileInfo(const FSoundFileDescription& FileDescription, const FString& FullPath) override;
+		ESoundFileError::Type SetEncodingQuality(double EncodingQuality) override;
+		ESoundFileError::Type BeginImport() override;
+		ESoundFileError::Type EndImport() override;
+		ESoundFileError::Type SeekFrames(SoundFileCount Offset, ESoundFileSeekMode::Type SeekMode, SoundFileCount& OutOffset) override;
+		ESoundFileError::Type ReadFrames(float* DataPtr, SoundFileCount NumFrames, SoundFileCount& OutNumFramesRead) override;
+		ESoundFileError::Type ReadFrames(double* DataPtr, SoundFileCount NumFrames, SoundFileCount& OutNumFramesRead) override;
+		ESoundFileError::Type ReadSamples(float* DataPtr, SoundFileCount NumSamples, SoundFileCount& OutNumSamplesRead) override;
+		ESoundFileError::Type ReadSamples(double* DataPtr, SoundFileCount NumSamples, SoundFileCount& OutNumSamplesRead) override;
+		ESoundFileError::Type WriteFrames(const float* Data, SoundFileCount NumFrames, SoundFileCount& OutNumFramesWritten) override;
+		ESoundFileError::Type WriteFrames(const double* Data, SoundFileCount NumFrames, SoundFileCount& OutNumFramesWritten) override;
+		ESoundFileError::Type WriteSamples(const float* DataPtr, SoundFileCount NumSamples, SoundFileCount& OutNumSampleWritten) override;
+		ESoundFileError::Type WriteSamples(const double* DataPtr, SoundFileCount NumSamples, SoundFileCount& OutNumSampleWritten) override;
+		
+		// ISoundFileVirtual implementation
+		ESoundFileError::Type GetLengthBytes(SoundFileCount& OutLength) const override;
+		ESoundFileError::Type SeekBytes(SoundFileCount Offset, ESoundFileSeekMode::Type SeekMode, SoundFileCount& OutOffset) override;
+		ESoundFileError::Type ReadBytes(void* DataPtr, SoundFileCount NumBytes, SoundFileCount& OutNumBytesRead) override;
+		ESoundFileError::Type WriteBytes(const void* DataPtr, SoundFileCount NumBytes, SoundFileCount& OutNumBytesWritten) override;
+		SoundFileCount GetOffsetBytes() const override;
+
+	private:
+
+		TSharedPtr<ISoundFileData> SoundFileData;
+
+		/** The current read/write index in the raw data */
+		SoundFileCount CurrentIndexBytes;
+
+		/** Handle to the sound file decoder. */
+		SoundFileHandle* FileHandle;
+
+		/** The current state of the sound file. */
+		volatile ESoundFileState::Type State;
+
+		/** The last error of the sound file. */
+		volatile ESoundFileError::Type CurrentError;
+
+		int32 Id;
+	};
+
+	/**
+	FSoundFileData
+	*/
+	class FSoundFileData : public ISoundFileData
+	{
+	public:
+		FSoundFileData(ESoundFileOpenMode::Type Mode);
+		~FSoundFileData();
+
+		// ISoundFileData interface
+		const FString& GetFilePath() const override;
+		void SetFilePath(const FString& InPath) override;
+		const FSoundFileDescription& GetDescription() const override;
+		void SetDescription(const FSoundFileDescription& InDescription) override;
+		const FSoundFileDescription& GetOriginalDescription() const override;
+		void SetOriginalDescription(const FSoundFileDescription& InDescription) override;
+		const TArray<ESoundFileChannelMap::Type>& GetChannelMap() const override;
+		void SetChannelMap(const TArray<ESoundFileChannelMap::Type>& InChannelMap) override;
+		const uint8* GetData() const override;
+		int64 GetDataSize() const override;
+		void SetData(uint8* InDataPtr, int64 InDataSize) override;
+
+		TArray<uint8>& GetDataArray();
+		void SetNumFrames(int64 InNumFrames);
+
+	private:
+		// The path of the file when it was imported
+		FString FilePath;
+
+		// The current description of the file (i.e. compressed)
+		FSoundFileDescription Description;
+
+		// The description of the file before it was imported
+		FSoundFileDescription OriginalDescription;
+
+		// The channel map of the file
+		TArray<ESoundFileChannelMap::Type> ChannelMap;
+
+		// An array of data, used to build the file data during import
+		TArray<uint8> DataArray;
+
+		// A ptr to a buffer of data, used to point this sound file data to an arbitrary location in memory
+		// This sound file data object does not own this data buffer
+		uint8* DataBuffer;
+
+		// The size of the data buffer in bytes
+		int64 DataBufferSize;
+
+		// The mode this sound file data is in (read or write)
+		ESoundFileOpenMode::Type Mode;
+	};
+}
+
+
+
+#endif // #if ENABLE_UNREAL_AUDIO
+
diff --git a/Engine/Source/Runtime/UnrealAudio/Public/UnrealAudioDeviceModule.h b/Engine/Source/Runtime/UnrealAudio/Public/UnrealAudioDeviceModule.h
--- a/Engine/Source/Runtime/UnrealAudio/Public/UnrealAudioDeviceModule.h
+++ b/Engine/Source/Runtime/UnrealAudio/Public/UnrealAudioDeviceModule.h
@@ -5,6 +5,7 @@
 #include "ModuleInterface.h"
 #include "UnrealAudioModule.h"
 #include "UnrealAudioDeviceFormat.h"
+#include "UnrealAudioTypes.h"
 
 #if ENABLE_UNREAL_AUDIO
 
@@ -17,75 +18,6 @@ DECLARE_LOG_CATEGORY_EXTERN(LogUnrealAudioDevice, Log, All);
 
 namespace UAudio
 {
-	/**
-	* EDeviceError
-	* The types of errors that can occur in IUnrealAudioDevice implementations.
-	*/
-	namespace EDeviceError
-	{
-		enum Type
-		{
-			WARNING,
-			UNKNOWN,
-			INVALID_PARAMETER,
-			PLATFORM,
-			SYSTEM,
-			THREAD,
-		};
-
-		/** @return the stringified version of the enum passed in */
-		static inline const TCHAR* ToString(EDeviceError::Type DeviceError)
-		{
-			switch (DeviceError)
-			{
-				case WARNING:			return TEXT("WARNING");
-				case UNKNOWN:			return TEXT("UNKNOWN");
-				case INVALID_PARAMETER:	return TEXT("INVALID_PARAMETER");
-				case PLATFORM:			return TEXT("PLATFORM");
-				case SYSTEM:			return TEXT("SYSTEM");
-				case THREAD:			return TEXT("THREAD");
-				default:				check(false);
-			}
-			return TEXT("");
-		}
-	}
-
-	/**
-	* EDeviceApi
-	* Implemented platform-Specific audio device APIs.
-	*/
-	namespace EDeviceApi
-	{
-		enum Type
-		{
-			WASAPI,				// Windows
-			XAUDIO2,			// Windows/XBox
-			NGS2,				// PS4
-			ALSA,				// Linux
-			CORE_AUDIO,			// Mac/IOS
-			OPENAL,				// Linux
-			HTML5,				// Web
-			DUMMY,				// DUMMY (no device api)
-		};
-
-		/** @return the stringified version of the enum passed in */
-		static inline const TCHAR* ToString(EDeviceApi::Type DeviceApi)
-		{
-			switch (DeviceApi)
-			{
-				case WASAPI:		return TEXT("WASAPI");
-				case XAUDIO2:		return TEXT("XAUDIO2");
-				case NGS2:			return TEXT("NGS2");
-				case ALSA:			return TEXT("ALSA");
-				case CORE_AUDIO:	return TEXT("CORE_AUDIO");
-				case OPENAL:		return TEXT("OPENAL");
-				case HTML5:			return TEXT("HTML5");
-				case DUMMY:			return TEXT("DUMMY");
-				default:			check(false);
-			}
-			return TEXT("");
-		}
-	};
 
 	/**
 	* EStreamStatus
@@ -115,66 +47,6 @@ namespace UAudio
 	}
 	typedef uint8 StreamStatus;
 
-	/**
-	* ESpeaker
-	* An enumeration to specify speaker types
-	*/
-	struct ESpeaker
-	{
-		enum Type
-		{
-			FRONT_LEFT,
-			FRONT_RIGHT,
-			FRONT_CENTER,
-			LOW_FREQUENCY,
-			BACK_LEFT,
-			BACK_RIGHT,
-			FRONT_LEFT_OF_CENTER,
-			FRONT_RIGHT_OF_CENTER,
-			BACK_CENTER,
-			SIDE_LEFT,
-			SIDE_RIGHT,
-			TOP_CENTER,
-			TOP_FRONT_LEFT,
-			TOP_FRONT_CENTER,
-			TOP_FRONT_RIGHT,
-			TOP_BACK_LEFT,
-			TOP_BACK_CENTER,
-			TOP_BACK_RIGHT,
-			UNUSED,
-			SPEAKER_TYPE_COUNT
-		};
-
-		/** @return the stringified version of the enum passed in */
-		static inline const TCHAR* ToString(ESpeaker::Type Speaker)
-		{
-			switch (Speaker)
-			{
-				case FRONT_LEFT:				return TEXT("FRONT_LEFT");
-				case FRONT_RIGHT:				return TEXT("FRONT_RIGHT");
-				case FRONT_CENTER:				return TEXT("FRONT_CENTER");
-				case LOW_FREQUENCY:				return TEXT("LOW_FREQUENCY");
-				case BACK_LEFT:					return TEXT("BACK_LEFT");
-				case BACK_RIGHT:				return TEXT("BACK_RIGHT");
-				case FRONT_LEFT_OF_CENTER:		return TEXT("FRONT_LEFT_OF_CENTER");
-				case FRONT_RIGHT_OF_CENTER:		return TEXT("FRONT_RIGHT_OF_CENTER");
-				case BACK_CENTER:				return TEXT("BACK_CENTER");
-				case SIDE_LEFT:					return TEXT("SIDE_LEFT");
-				case SIDE_RIGHT:				return TEXT("SIDE_RIGHT");
-				case TOP_CENTER:				return TEXT("TOP_CENTER");
-				case TOP_FRONT_LEFT:			return TEXT("TOP_FRONT_LEFT");
-				case TOP_FRONT_CENTER:			return TEXT("TOP_FRONT_CENTER");
-				case TOP_FRONT_RIGHT:			return TEXT("TOP_FRONT_RIGHT");
-				case TOP_BACK_LEFT:				return TEXT("TOP_BACK_LEFT");
-				case TOP_BACK_CENTER:			return TEXT("TOP_BACK_CENTER");
-				case TOP_BACK_RIGHT:			return TEXT("TOP_BACK_RIGHT");
-				case UNUSED:					return TEXT("UNUSED");
-				default:						return TEXT("UKNOWN");
-			}
-			return TEXT("");
-		}
-	};
-
 	/**
 	Static array specifying which frame rates to check for support on audio devices.
 	http://en.wikipedia.org/wiki/Sampling_%28signal_processing%29
diff --git a/Engine/Source/Runtime/UnrealAudio/Public/UnrealAudioModule.h b/Engine/Source/Runtime/UnrealAudio/Public/UnrealAudioModule.h
--- a/Engine/Source/Runtime/UnrealAudio/Public/UnrealAudioModule.h
+++ b/Engine/Source/Runtime/UnrealAudio/Public/UnrealAudioModule.h
@@ -3,13 +3,8 @@
 #pragma once 
 
 #include "ModuleInterface.h"
-
-// Only enable unreal audio on windows or mac
-#if PLATFORM_WINDOWS || PLATFORM_MAC
-#define ENABLE_UNREAL_AUDIO 1
-#else
-#define ENABLE_UNREAL_AUDIO 0
-#endif
+#include "UnrealAudioTypes.h"
+#include "UnrealAudioSoundFile.h"
 
 namespace UAudio
 {
@@ -31,6 +26,29 @@ namespace UAudio
 		virtual void Shutdown()
 		{
 		}
+
+		/** 
+		ImportSound
+		@param ImportSettings Settings to use to import a sound file. Settings file contains all information
+				needed to performt he import process.
+		@return A shared pointer to an ISoundFile object.
+		*/
+		virtual	TSharedPtr<ISoundFile> ImportSound(const FSoundFileImportSettings& ImportSettings)
+		{
+			return nullptr;
+		}
+
+		/**
+		ImportSound
+		@param ImportSettings Settings to use to import a sound file. Settings file contains all information
+		needed to performt he import process.
+		@return A handle to an ISoundFile object.
+		*/
+		virtual	void ExportSound(TSharedPtr<ISoundFileData> SoundFileData, const FString& ExportPath)
+		{
+		}
+
+
 	};
 }
 
diff --git a/Engine/Source/Runtime/UnrealAudio/Public/UnrealAudioSoundFile.h b/Engine/Source/Runtime/UnrealAudio/Public/UnrealAudioSoundFile.h
new file mode 100644
--- /dev/null
+++ b/Engine/Source/Runtime/UnrealAudio/Public/UnrealAudioSoundFile.h
@@ -0,0 +1,382 @@
+// Copyright 1998-2015 Epic Games, Inc. All Rights Reserved.
+
+#pragma once 
+
+#include "UnrealAudioModule.h"
+#include "UnrealAudioTypes.h"
+
+namespace UAudio
+{
+	/*
+	* ESoundFileError
+	**/
+	namespace ESoundFileError
+	{
+		enum Type
+		{
+			NONE,
+			INVALID_SOUND_FILE,
+			BAD_ENCODING_QUALITY,
+			ALREADY_OPENED,
+			ALREADY_HAS_DATA,
+			INVALID_DATA,
+			FILE_DOESNT_EXIST,
+			INVALID_INPUT_FORMAT,
+			INVALID_CHANNEL_MAP,
+			FAILED_TO_OPEN,
+			FAILED_TO_SEEK,
+			ALREADY_INTIALIZED,
+			LOADING,
+			UNKNOWN
+		};
+
+		/** @return the stringified version of the enum passed in */
+		inline const TCHAR* ToString(ESoundFileError::Type SoundFileError)
+		{
+			switch (SoundFileError)
+			{
+				case NONE:					return TEXT("NONE");
+				case INVALID_SOUND_FILE:	return TEXT("INVALID_SOUND_FILE");
+				case BAD_ENCODING_QUALITY:	return TEXT("BAD_ENCODING_QUALITY");
+				case ALREADY_OPENED:		return TEXT("ALREADY_OPENED");
+				case ALREADY_HAS_DATA:		return TEXT("ALREADY_HAS_DATA");
+				case INVALID_DATA:			return TEXT("INVALID_DATA");
+				case FILE_DOESNT_EXIST:		return TEXT("FILE_DOESNT_EXIST");
+				case INVALID_INPUT_FORMAT:	return TEXT("INVALID_INPUT_FORMAT");
+				case INVALID_CHANNEL_MAP:	return TEXT("INVALID_CHANNEL_MAP");
+				case FAILED_TO_OPEN:		return TEXT("FAILED_TO_OPEN");
+				case FAILED_TO_SEEK:		return TEXT("FAILED_TO_SEEK");
+				case ALREADY_INTIALIZED:	return TEXT("ALREADY_INTIALIZED");
+				case LOADING:				return TEXT("LOADING");
+
+				default:
+				case UNKNOWN:				return TEXT("UNKNOWN");
+			}
+		}
+	}
+
+	/*
+	* ESoundFileFormatFlags
+	* Specifies the major format type of the sound source.
+	* File formats are fully specified by a major/minor format.
+	*
+	* For example, a Ogg-Vorbis encoding would use:
+	* uint32 FormatFlags = ESoundFormatFlags::OGG | ESoundFormatFlags::VORBIS;
+	**/
+	namespace ESoundFileFormat
+	{
+		enum Flags
+		{
+			// Major Formats
+			WAV				= 0x010000,		// Microsoft WAV format
+			AIFF			= 0x020000,		// Apple AIFF format
+			FLAC			= 0x170000,		// FLAC lossless
+			OGG				= 0x200000,		// Xiph OGG
+
+			// Uncompressed Minor Formats
+			PCM_SIGNED_8	= 0x0001,		// Signed 8 bit PCM
+			PCM_SIGNED_16	= 0x0002,		// Signed 16 bit PCM
+			PCM_SIGNED_24	= 0x0003,		// Signed 24 bit PCM
+			PCM_SIGNED_32	= 0x0004,		// Signed 32 bit PCM
+			PCM_UNSIGNED_8	= 0x0005,		// Unsigned 8 bit PCM
+			PCM_FLOAT		= 0x0006,		// 32 bit float
+			PCM_DOUBLE		= 0x0007,		// 64 bit float
+
+			// Compressed Minor Formats
+			MU_LAW			= 0x0010,		// Mu-law encoding
+			A_LAW			= 0x0011,		// A-law encoding
+			IMA_ADPCM		= 0x0012,		// IMA ADPCM encoding
+			MS_ADPCM		= 0x0013,		// Microsoft ADPCM encoding
+			GSM_610			= 0x0020,		// GSM 6.10 encoding
+			G721_32			= 0x0030,		// 32 kbps G721 ADPCM encoding
+			G723_24			= 0x0031,		// 23 kbps G723 ADPCM encoding
+			G723_40			= 0x0032,		// 40 kbps G723 ADPCM encoding
+			DWVW_12			= 0x0040,		// 12 bit delta-width variable word encoding
+			DMVW_16			= 0x0041,		// 16 bit delta-width variable word encoding
+			DMVW_24			= 0x0042,		// 24 bit delta-width variable word encoding
+			DMVW_N			= 0x0043,		// N bit delta-width variable word encoding
+			VORBIS			= 0x0060,		// Xiph vorbis encoding
+
+			// Endian opts
+			ENDIAN_FILE		= 0x00000000,	// default file endian
+			ENDIAN_LITTLE	= 0x10000000,	// little-endian
+			ENDIAN_BIG		= 0x20000000,	// big-endian
+			ENDIAN_CPU		= 0x30000000,	// cpu-endian
+
+			// Masks
+			MINOR_FORMAT_MASK = 0x0000FFFF,
+			MAJOR_FORMAT_MASK = 0x0FFF0000,
+			ENDIAN_MASK		  = 0x30000000,
+		};
+
+		inline const TCHAR* ToStringMajor(int32 FormatFlags)
+		{
+			switch (FormatFlags & ESoundFileFormat::MAJOR_FORMAT_MASK)
+			{
+				case ESoundFileFormat::WAV:		return TEXT("WAV");
+				case ESoundFileFormat::AIFF:	return TEXT("AIFF");
+				case ESoundFileFormat::FLAC:	return TEXT("FLAC");
+				case ESoundFileFormat::OGG:		return TEXT("OGG");
+				default:						return TEXT("INVALID");
+			}
+		}
+
+		inline const TCHAR* ToStringMinor(int32 FormatFlags)
+		{
+			switch (FormatFlags & ESoundFileFormat::MINOR_FORMAT_MASK)
+			{
+				case ESoundFileFormat::PCM_SIGNED_8:	return TEXT("PCM_SIGNED_8");
+				case ESoundFileFormat::PCM_SIGNED_16:	return TEXT("PCM_SIGNED_16");
+				case ESoundFileFormat::PCM_SIGNED_24:	return TEXT("PCM_SIGNED_24");
+				case ESoundFileFormat::PCM_SIGNED_32:	return TEXT("PCM_SIGNED_32");
+				case ESoundFileFormat::PCM_UNSIGNED_8:	return TEXT("PCM_UNSIGNED_8");
+				case ESoundFileFormat::PCM_FLOAT:		return TEXT("PCM_FLOAT");
+				case ESoundFileFormat::PCM_DOUBLE:		return TEXT("PCM_DOUBLE");
+				case ESoundFileFormat::MU_LAW:			return TEXT("MU_LAW");
+				case ESoundFileFormat::A_LAW:			return TEXT("A_LAW");
+				case ESoundFileFormat::IMA_ADPCM:		return TEXT("IMA_ADPCM");
+				case ESoundFileFormat::MS_ADPCM:		return TEXT("MS_ADPCM");
+				case ESoundFileFormat::GSM_610:			return TEXT("GSM_610");
+				case ESoundFileFormat::G721_32:			return TEXT("G721_32");
+				case ESoundFileFormat::G723_24:			return TEXT("G723_24");
+				case ESoundFileFormat::G723_40:			return TEXT("G723_40");
+				case ESoundFileFormat::DWVW_12:			return TEXT("DWVW_12");
+				case ESoundFileFormat::DMVW_16:			return TEXT("DMVW_16");
+				case ESoundFileFormat::DMVW_24:			return TEXT("DMVW_24");
+				case ESoundFileFormat::DMVW_N:			return TEXT("DMVW_N");
+				case ESoundFileFormat::VORBIS:			return TEXT("VORBIS");
+				default:								return TEXT("INVALID");
+			}
+		}
+	};
+
+	/*
+	* ESoundFileChannelMap
+	* These are separated from the device channel mappings purposefully since
+	* the enumeration may not exactly be the same as the output speaker mapping.
+	**/
+	namespace ESoundFileChannelMap
+	{
+		enum Type
+		{
+			INVALID = 0,
+			MONO,
+			LEFT,
+			RIGHT,
+			CENTER,
+			FRONT_LEFT,
+			FRONT_RIGHT,
+			FRONT_CENTER,
+			BACK_CENTER,
+			BACK_LEFT,
+			BACK_RIGHT,
+			LFE,
+			LEFT_CENTER,
+			RIGHT_CENTER,
+			SIDE_LEFT,
+			SIDE_RIGHT,
+			TOP_CENTER,
+			TOP_FRONT_LEFT,
+			TOP_FRONT_RIGHT,
+			TOP_FRONT_CENTER,
+			TOP_BACK_LEFT,
+			TOP_BACK_RIGHT,
+			TOP_BACK_CENTER,
+		};
+
+		/** @return the stringified version of the enum passed in */
+		inline const TCHAR* ToString(ESoundFileChannelMap::Type ChannelMap)
+		{
+			switch (ChannelMap)
+			{
+				case INVALID:			return TEXT("INVALID");
+				case MONO:				return TEXT("MONO");
+				case LEFT:				return TEXT("LEFT");
+				case RIGHT:				return TEXT("RIGHT");
+				case CENTER:			return TEXT("CENTER");
+				case FRONT_LEFT:		return TEXT("FRONT_LEFT");
+				case FRONT_RIGHT:		return TEXT("FRONT_RIGHT");
+				case FRONT_CENTER:		return TEXT("FRONT_CENTER");
+				case BACK_CENTER:		return TEXT("BACK_CENTER");
+				case BACK_LEFT:			return TEXT("BACK_LEFT");
+				case BACK_RIGHT:		return TEXT("BACK_RIGHT");
+				case LFE:				return TEXT("LFE");
+				case LEFT_CENTER:		return TEXT("LEFT_CENTER");
+				case RIGHT_CENTER:		return TEXT("RIGHT_CENTER");
+				case SIDE_LEFT:			return TEXT("SIDE_LEFT");
+				case SIDE_RIGHT:		return TEXT("SIDE_RIGHT");
+				case TOP_CENTER:		return TEXT("TOP_CENTER");
+				case TOP_FRONT_LEFT:	return TEXT("TOP_FRONT_LEFT");
+				case TOP_FRONT_RIGHT:	return TEXT("TOP_FRONT_RIGHT");
+				case TOP_FRONT_CENTER:	return TEXT("TOP_FRONT_CENTER");
+				case TOP_BACK_LEFT:		return TEXT("TOP_BACK_LEFT");
+				case TOP_BACK_RIGHT:	return TEXT("TOP_BACK_RIGHT");
+				case TOP_BACK_CENTER:	return TEXT("TOP_BACK_CENTER");
+				default:				return TEXT("UNKNOWN");
+			}
+		}
+	}
+
+	/*
+	* FSoundFileDescription
+	* Struct specifying a sound file description.
+	**/
+	struct FSoundFileDescription
+	{
+		// The number of frames (interleaved samples) in the sound file.
+		int64 NumFrames;
+
+		// The sample rate of the sound file.
+		int32 SampleRate;
+
+		// The number of channels of the sound file.
+		int32 NumChannels;
+
+		// The format flags of the sound file.
+		int32 FormatFlags;
+
+		// The number of sections of the sound file.
+		int32 NumSections;
+
+		// Whether or not the sound file is seekable.
+		int32 bIsSeekable;
+	};
+
+	/**
+	* FSoundFileImportSettings
+	* Struct defining sound file import settings.
+	*/
+	struct UNREALAUDIO_API FSoundFileImportSettings
+	{
+		// Full path to the sound file on disk.
+		FString SoundFilePath;
+
+		// Desired import format.
+		int32 Format;
+
+		// Desired import sample rate.
+		uint32 SampleRate;
+
+		// For compression-type target formats that used an encoding quality (0.0 = low, 1.0 = high).
+		double EncodingQuality;
+
+		// Whether or not to peak-normalize the audio file during import.
+		bool bPerformPeakNormalization;
+	};
+
+	/*
+	* ESoundFileState
+	**/
+	namespace ESoundFileState
+	{
+		enum Type
+		{
+			// Sound file hasn't begun loading and was just created.
+			UNINITIALIZED,
+
+			// Sound file has been initialized with data and is ready for playback
+			INITIALIZED,
+
+			// Sound file is importing.
+			IMPORTING,
+
+			// Sound file has finished importing.
+			IMPORTED,
+
+			// Sound file has an error. Called GetError() on ISoundFile to get error.
+			HAS_ERROR,
+		};
+	}
+
+	/**
+	* FSoundFileData
+	* Data structure containing data for a sound file. Can be stored within a UAsset or loaded with a sound file.
+	*/
+	class UNREALAUDIO_API ISoundFileData
+	{
+	public:
+		virtual ~ISoundFileData() {}
+
+		/** Sets and gets the sound file path */
+		virtual const FString& GetFilePath() const = 0;
+		virtual void SetFilePath(const FString& InPath) = 0;
+		
+		/** Sets and gets the file description */
+		virtual const FSoundFileDescription& GetDescription() const = 0;
+		virtual void SetDescription(const FSoundFileDescription& InDescription) = 0;
+
+		/** Sets and gets the original file description */
+		virtual const FSoundFileDescription& GetOriginalDescription() const = 0;
+		virtual void SetOriginalDescription(const FSoundFileDescription& InDescription) = 0;
+
+		/** Sets and gets the channel map */
+		virtual const TArray<ESoundFileChannelMap::Type>& GetChannelMap() const = 0;
+		virtual void SetChannelMap(const TArray<ESoundFileChannelMap::Type>& InChannelMap) = 0;
+
+		/** Gets a const pointer to the internal buffer of data. */
+		virtual const uint8* GetData() const = 0;
+
+		/** Gets the size of the data */
+		virtual int64 GetDataSize() const = 0;
+
+		/** Sets the data buffer pointer and the size of the data buffer */
+		virtual void SetData(uint8* InDataPtr, int64 InDataSize) = 0;
+	};
+
+	/**
+	* ISoundFile
+	* Interface for sound file reading/writing.
+	*/
+	class UNREALAUDIO_API ISoundFile
+	{
+	public:
+		/** Virtual Destructor */
+		virtual ~ISoundFile() {}
+
+		/** Initializes the sound file with the given sound file data. */
+		virtual ESoundFileError::Type Initialize(TSharedPtr<ISoundFileData> InSoundFileData) = 0;
+
+		/** Returns the name of the sound file */
+		virtual ESoundFileError::Type GetName(FString& Name) = 0;
+
+		/** Gets the sound file data associated with this ISoundFile. */
+		virtual ESoundFileError::Type GetSoundFileData(TSharedPtr<ISoundFileData>& OutSoundFileData) = 0;
+
+		/** @return Returns the state of the sound file. */
+		virtual  ESoundFileState::Type GetState() const = 0;
+
+		/** @return Returns the last sound file error. */
+		virtual ESoundFileError::Type GetError() const = 0;
+	};
+
+	/**
+	* Creates an empty sound file data object
+	* Creates an empty sound file (with no sound file data)
+	*/
+	UNREALAUDIO_API TSharedPtr<ISoundFileData> CreateSoundFileData();
+
+	/**
+	* CreateSoundFile
+	* Creates an empty sound file (with no sound file data)
+	*/
+	UNREALAUDIO_API TSharedPtr<ISoundFile> CreateSoundFile();
+
+	/**
+	* ImportSound
+	* Imports a sound given the import settings.
+	* The function immediately returns a ISoundFile reference, but that sound file will
+	* be in a loading state. If the sound file succeeds in importing, the sound file state
+	* will become Loaded, otherwise it will have an error that can be retrieved.
+	*/
+	UNREALAUDIO_API TSharedPtr<ISoundFile> ImportSound(const FSoundFileImportSettings& ImportSettings);
+
+	/**
+	* ExportSound
+	* Exports a sound to the given path.
+	* The function immediately returns a ISoundFile reference, but that sound file will
+	* be in a loading state. If the sound file succeeds in importing, the sound file state
+	* will become Loaded, otherwise it will have an error that can be retrieved.
+	*/
+	UNREALAUDIO_API void ExportSound(TSharedPtr<ISoundFileData> SoundFileData, const FString& ExportPath);
+
+}
diff --git a/Engine/Source/Runtime/UnrealAudio/Public/UnrealAudioTests.h b/Engine/Source/Runtime/UnrealAudio/Public/UnrealAudioTests.h
--- a/Engine/Source/Runtime/UnrealAudio/Public/UnrealAudioTests.h
+++ b/Engine/Source/Runtime/UnrealAudio/Public/UnrealAudioTests.h
@@ -43,5 +43,33 @@ namespace UAudio
 	*/
 	bool UNREALAUDIO_API TestDeviceOutputNoisePan(double TestTime);
 
+	/**
+	* TestSourceImport
+	*
+	* Tests importing a sound file.
+	* @param ImportSettings A struct defining import settings.
+	* @return True if test succeeds.
+	*/
+	bool UNREALAUDIO_API TestSourceImport(const FSoundFileImportSettings& ImportSettings);
+
+	/**
+	* TestSourceImport
+	*
+	* Tests importing a sound file, then exporting it (with "export" appended to filename).
+	* @param ImportSettings A struct defining import settings.
+	* @return True if test succeeds.
+	*/
+	bool UNREALAUDIO_API TestSourceImportExport(const FSoundFileImportSettings& ImportSettings);
+
+	/**
+	* TestSourceImportFolder
+	*
+	* Tests importing all the sound files in a given folder.
+	* @param FolderPath Path to a folder of sound files.
+	* @param ImportSettings Settings to use for each sound file import (path is ignored)
+	* @return True if test succeeds.
+	*/
+	bool UNREALAUDIO_API TestSourceImportFolder(FString& FolderPath, const FSoundFileImportSettings& ImportSettings);
+
 }
 
diff --git a/Engine/Source/Runtime/UnrealAudio/Public/UnrealAudioTypes.h b/Engine/Source/Runtime/UnrealAudio/Public/UnrealAudioTypes.h
new file mode 100644
--- /dev/null
+++ b/Engine/Source/Runtime/UnrealAudio/Public/UnrealAudioTypes.h
@@ -0,0 +1,152 @@
+// Copyright 1998-2015 Epic Games, Inc. All Rights Reserved.
+
+#pragma once 
+
+// Only enable unreal audio on windows or mac
+#if PLATFORM_WINDOWS || PLATFORM_MAC
+#define ENABLE_UNREAL_AUDIO 1
+#else
+#define ENABLE_UNREAL_AUDIO 0
+#endif
+
+#include "UnrealAudioModule.h"
+
+#if ENABLE_UNREAL_AUDIO
+
+namespace UAudio
+{
+	/**
+	* EDeviceError
+	* The types of errors that can occur in IUnrealAudioDevice implementations.
+	*/
+	namespace EDeviceError
+	{
+		enum Type
+		{
+			WARNING,
+			UNKNOWN,
+			INVALID_PARAMETER,
+			PLATFORM,
+			SYSTEM,
+			THREAD,
+		};
+
+		/** @return the stringified version of the enum passed in */
+		static inline const TCHAR* ToString(EDeviceError::Type DeviceError)
+		{
+			switch (DeviceError)
+			{
+				case WARNING:			return TEXT("WARNING");
+				case UNKNOWN:			return TEXT("UNKNOWN");
+				case INVALID_PARAMETER:	return TEXT("INVALID_PARAMETER");
+				case PLATFORM:			return TEXT("PLATFORM");
+				case SYSTEM:			return TEXT("SYSTEM");
+				case THREAD:			return TEXT("THREAD");
+				default:				check(false);
+			}
+			return TEXT("");
+		}
+	}
+
+	/**
+	* EDeviceApi
+	* Implemented platform-Specific audio device APIs.
+	*/
+	namespace EDeviceApi
+	{
+		enum Type
+		{
+			WASAPI,				// Windows
+			XAUDIO2,			// Windows/XBox
+			NGS2,				// PS4
+			ALSA,				// Linux
+			CORE_AUDIO,			// Mac/IOS
+			OPENAL,				// Linux
+			HTML5,				// Web
+			DUMMY,				// DUMMY (no device api)
+		};
+
+		/** @return the stringified version of the enum passed in */
+		static inline const TCHAR* ToString(EDeviceApi::Type DeviceApi)
+		{
+			switch (DeviceApi)
+			{
+			case WASAPI:		return TEXT("WASAPI");
+			case XAUDIO2:		return TEXT("XAUDIO2");
+			case NGS2:			return TEXT("NGS2");
+			case ALSA:			return TEXT("ALSA");
+			case CORE_AUDIO:	return TEXT("CORE_AUDIO");
+			case OPENAL:		return TEXT("OPENAL");
+			case HTML5:			return TEXT("HTML5");
+			case DUMMY:			return TEXT("DUMMY");
+			default:			check(false);
+			}
+			return TEXT("");
+		}
+	};
+
+	/**
+	* ESpeaker
+	* An enumeration to specify speaker types
+	*/
+	namespace ESpeaker
+	{
+		enum Type
+		{
+			FRONT_LEFT,
+			FRONT_RIGHT,
+			FRONT_CENTER,
+			LOW_FREQUENCY,
+			BACK_LEFT,
+			BACK_RIGHT,
+			FRONT_LEFT_OF_CENTER,
+			FRONT_RIGHT_OF_CENTER,
+			BACK_CENTER,
+			SIDE_LEFT,
+			SIDE_RIGHT,
+			TOP_CENTER,
+			TOP_FRONT_LEFT,
+			TOP_FRONT_CENTER,
+			TOP_FRONT_RIGHT,
+			TOP_BACK_LEFT,
+			TOP_BACK_CENTER,
+			TOP_BACK_RIGHT,
+			UNUSED,
+			SPEAKER_TYPE_COUNT
+		};
+
+		/** @return the stringified version of the enum passed in */
+		static inline const TCHAR* ToString(ESpeaker::Type Speaker)
+		{
+			switch (Speaker)
+			{
+			case FRONT_LEFT:				return TEXT("FRONT_LEFT");
+			case FRONT_RIGHT:				return TEXT("FRONT_RIGHT");
+			case FRONT_CENTER:				return TEXT("FRONT_CENTER");
+			case LOW_FREQUENCY:				return TEXT("LOW_FREQUENCY");
+			case BACK_LEFT:					return TEXT("BACK_LEFT");
+			case BACK_RIGHT:				return TEXT("BACK_RIGHT");
+			case FRONT_LEFT_OF_CENTER:		return TEXT("FRONT_LEFT_OF_CENTER");
+			case FRONT_RIGHT_OF_CENTER:		return TEXT("FRONT_RIGHT_OF_CENTER");
+			case BACK_CENTER:				return TEXT("BACK_CENTER");
+			case SIDE_LEFT:					return TEXT("SIDE_LEFT");
+			case SIDE_RIGHT:				return TEXT("SIDE_RIGHT");
+			case TOP_CENTER:				return TEXT("TOP_CENTER");
+			case TOP_FRONT_LEFT:			return TEXT("TOP_FRONT_LEFT");
+			case TOP_FRONT_CENTER:			return TEXT("TOP_FRONT_CENTER");
+			case TOP_FRONT_RIGHT:			return TEXT("TOP_FRONT_RIGHT");
+			case TOP_BACK_LEFT:				return TEXT("TOP_BACK_LEFT");
+			case TOP_BACK_CENTER:			return TEXT("TOP_BACK_CENTER");
+			case TOP_BACK_RIGHT:			return TEXT("TOP_BACK_RIGHT");
+			case UNUSED:					return TEXT("UNUSED");
+			default:						return TEXT("UKNOWN");
+			}
+			return TEXT("");
+		}
+	};
+
+}
+
+
+#endif // #if ENABLE_UNREAL_AUDIO
+
diff --git a/Engine/Source/Runtime/UnrealAudio/UnrealAudio.Build.cs b/Engine/Source/Runtime/UnrealAudio/UnrealAudio.Build.cs
--- a/Engine/Source/Runtime/UnrealAudio/UnrealAudio.Build.cs
+++ b/Engine/Source/Runtime/UnrealAudio/UnrealAudio.Build.cs
@@ -31,5 +31,27 @@ public class UnrealAudio : ModuleRules
 				"CoreUObject", 
 			}
 		);
+
+		// Libsndfile DLL
+		string LibSndFilePath = UEBuildConfiguration.UEThirdPartyBinariesDirectory + "libsndfile/";
+		if (Target.Platform == UnrealTargetPlatform.Win32)
+		{
+			LibSndFilePath += "Win32";
+			PublicLibraryPaths.Add(LibSndFilePath);
+			PublicAdditionalLibraries.Add("libsndfile-1.lib");
+			PublicDelayLoadDLLs.Add("libsndfile-1.dll");
+		}
+		else if (Target.Platform == UnrealTargetPlatform.Win64)
+		{
+			LibSndFilePath += "Win64";
+			PublicLibraryPaths.Add(LibSndFilePath);
+			PublicAdditionalLibraries.Add("libsndfile-1.lib");
+			PublicDelayLoadDLLs.Add("libsndfile-1.dll");
+		}
+		else if (Target.Platform == UnrealTargetPlatform.Mac)
+		{
+			LibSndFilePath += "Mac/libsndfile.1.dylib";
+			PublicAdditionalLibraries.Add(LibSndFilePath);
+		}
 	}
 }
\ No newline at end of file
