UnrealAudio

Disabling mac support temporarily to resolve a missing dylib on the build machines.

diff --git a/Engine/Source/Runtime/UnrealAudio/Public/UnrealAudioTypes.h b/Engine/Source/Runtime/UnrealAudio/Public/UnrealAudioTypes.h
--- a/Engine/Source/Runtime/UnrealAudio/Public/UnrealAudioTypes.h
+++ b/Engine/Source/Runtime/UnrealAudio/Public/UnrealAudioTypes.h
@@ -3,7 +3,7 @@
 #pragma once 
 
 // Only enable unreal audio on windows or mac
-#if PLATFORM_WINDOWS || PLATFORM_MAC
+#if PLATFORM_WINDOWS
 #define ENABLE_UNREAL_AUDIO 1
 #else
 #define ENABLE_UNREAL_AUDIO 0
diff --git a/Engine/Source/Runtime/UnrealAudio/UnrealAudio.Build.cs b/Engine/Source/Runtime/UnrealAudio/UnrealAudio.Build.cs
--- a/Engine/Source/Runtime/UnrealAudio/UnrealAudio.Build.cs
+++ b/Engine/Source/Runtime/UnrealAudio/UnrealAudio.Build.cs
@@ -48,10 +48,10 @@ public class UnrealAudio : ModuleRules
 			PublicAdditionalLibraries.Add("libsndfile-1.lib");
 			PublicDelayLoadDLLs.Add("libsndfile-1.dll");
 		}
-		else if (Target.Platform == UnrealTargetPlatform.Mac)
-		{
-			LibSndFilePath += "Mac/libsndfile.1.dylib";
-			PublicAdditionalLibraries.Add(LibSndFilePath);
-		}
+// 		else if (Target.Platform == UnrealTargetPlatform.Mac)
+// 		{
+// 			LibSndFilePath += "Mac/libsndfile.1.dylib";
+// 			PublicAdditionalLibraries.Add(LibSndFilePath);
+// 		}
 	}
 }
\ No newline at end of file
