Unreal Audio WIP
* Restructured sound file reading/writing classes
    - 2 classes: a sound file reader and a sound file writer
    - sound file reader decodes sound file data for supported formats, can read stream directly off disk
    - sound file writer writes to a bulk data buffer, can write for supported formats
* First pass sound file asset manager
    - has target memory usage for loaded sound files
    - flushes unrefed sound files when usage is above threshold
    - flushes sound files that haven't been used above a time threshold (e.g. 30 seconds)
    - supports sync and async loading from disk
    - supports loads from bulk data
    - some tests which simulate randomized real-game load patterns.
    - rewrote tests for batch asset conversion (from X to Y audio formats, including sample rates, see below)
* Implemented a flexible linear-based sample rate converter
    - Can be used for real-time decoding and for sample format conversion
* First pass voice manager
    - supports real and virtual voices and logic to steal real voices based on volume-weighted priority
    - organized to attempt to maximize cache coherency
    - manages messages from main-thread to audio-system thread to update voice params
* Sketches on low level mixer, not yet fully implemented
    - first pass implementation on async sound file decoding for "real" voices
* First pass on thread message system between main/audio threads and audio thread to device thread.
    - plan to switch to "named thread" system used by rendering thread for main->audio system thread communication
* First pass on emitter class
    - listener class not implemented, this is testing creating emitters and passing emitter params to audio system thread
* Couple utility classes useful for thread safe message passing
    - simple lockless single producer/single consumer dynamic queue for message passing
    - multi-producer/multi-consumer dynamic queue for message passing (not used)
    - simple class to check for thread ids to ensure functions are called on correct threads

diff --git a/Engine/Source/Editor/UnrealEd/Private/Commandlets/AudioTestCommandlet.cpp b/Engine/Source/Editor/UnrealEd/Private/Commandlets/AudioTestCommandlet.cpp
--- a/Engine/Source/Editor/UnrealEd/Private/Commandlets/AudioTestCommandlet.cpp
+++ b/Engine/Source/Editor/UnrealEd/Private/Commandlets/AudioTestCommandlet.cpp
@@ -84,38 +84,44 @@ static bool TestAudioDeviceOutputPan(const TArray<FString>& Args)
 	return UAudio::TestDeviceOutputNoisePan(-1);
 }
 
-static bool TestAudioSourceImport(const TArray<FString>& Args)
-{
-	if (Args.Num() != 1)
+static bool TestAudioSourceConvert(const TArray<FString>& Args)
+{	if (Args.Num() != 1)
 	{
 		return false;
 	}
 
-	UAudio::FSoundFileImportSettings SoundFileImportSettings;
-	SoundFileImportSettings.SoundFilePath = Args[0];
-	SoundFileImportSettings.Format = UAudio::ESoundFileFormat::OGG | UAudio::ESoundFileFormat::VORBIS;
-	SoundFileImportSettings.SampleRate = 48000;
-	SoundFileImportSettings.EncodingQuality = 0.5;
-	SoundFileImportSettings.bPerformPeakNormalization = true;
+	UAudio::FSoundFileConvertFormat SoundFileLoadSettings;
+	SoundFileLoadSettings.Format = UAudio::ESoundFileFormat::OGG | UAudio::ESoundFileFormat::VORBIS;
+	SoundFileLoadSettings.SampleRate = 48000;
+	SoundFileLoadSettings.EncodingQuality = 0.5;
+	SoundFileLoadSettings.bPerformPeakNormalization = false;
+
+	return UAudio::TestSourceConvert(Args[0], SoundFileLoadSettings);
+}
 
-	return UAudio::TestSourceImport(SoundFileImportSettings);
+static bool TestAudioSystemEmitterManager(const TArray<FString>& Args)
+{
+	return UAudio::TestEmitterManager();
 }
 
-static bool TestAudioSourceImportExport(const TArray<FString>& Args)
+static bool TestAudioSystemVoiceManager(const TArray<FString>& Args)
 {
 	if (Args.Num() != 1)
 	{
 		return false;
 	}
 
-	UAudio::FSoundFileImportSettings SoundFileImportSettings;
-	SoundFileImportSettings.SoundFilePath = Args[0];
-	SoundFileImportSettings.Format = UAudio::ESoundFileFormat::OGG | UAudio::ESoundFileFormat::VORBIS;
-	SoundFileImportSettings.SampleRate = 48000;
-	SoundFileImportSettings.EncodingQuality = 0.5;
-	SoundFileImportSettings.bPerformPeakNormalization = false;
+	return UAudio::TestVoiceManager(Args[0]);
+}
+
+static bool TestAudioSystemSoundFileManager(const TArray<FString>& Args)
+{
+	if (Args.Num() != 1)
+	{
+		return false;
+	}
 
-	return UAudio::TestSourceImportExport(SoundFileImportSettings);
+	return UAudio::TestSoundFileManager(Args[0]);
 }
 
 /**
@@ -135,25 +141,29 @@ struct AudioTestInfo
 
 enum EAudioTests
 {
-	AUDIO_TEST_DEVICE_ALL,
+	AUDIO_TEST_DEVICE_ALL = 0,
 	AUDIO_TEST_DEVICE_QUERY,
 	AUDIO_TEST_DEVICE_OUTPUT_SIMPLE,
 	AUDIO_TEST_DEVICE_OUTPUT_FM,
 	AUDIO_TEST_DEVICE_OUTPUT_PAN,
-	AUDIO_TEST_SOURCE_IMPORT,
-	AUDIO_TEST_SOURCE_IMPORT_EXPORT,
+	AUDIO_TEST_SOURCE_CONVERT,
+	AUDIO_TEST_SYSTEM_EMITTER_MANAGER,
+	AUDIO_TEST_SYSTEM_VOICE_MANAGER,
+	AUDIO_TEST_SYSTEM_SOUNDFILE_MANAGER,
 	AUDIO_TESTS
 };
 
 static AudioTestInfo AudioTestInfoList[] =
 {
-	{ "device", "all",		"None",						0, TestAudioDeviceAll			}, // AUDIO_TEST_DEVICE_ALL
-	{ "device", "query",	"None",						0, TestAudioDeviceQuery			}, // AUDIO_TEST_DEVICE_QUERY
-	{ "device", "out",		"None",						0, TestAudioDeviceOutputSimple	}, // AUDIO_TEST_DEVICE_OUTPUT_SIMPLE
-	{ "device", "out_fm",	"None",						0, TestAudioDeviceOutputFm		}, // AUDIO_TEST_DEVICE_OUTPUT_FM
-	{ "device", "out_pan",	"None",						0, TestAudioDeviceOutputPan		}, // AUDIO_TEST_DEVICE_OUTPUT_PAN
-	{ "source", "import",	"SourcePath",				1, TestAudioSourceImport		}, // AUDIO_TEST_SOURCE_IMPORT
-	{ "source", "export",	"SourcePath or Directory",	1, TestAudioSourceImportExport	}, // AUDIO_TEST_SOURCE_IMPORT_EXPORT
+	{ "device", "all",				"None",						0, TestAudioDeviceAll				}, // AUDIO_TEST_DEVICE_ALL
+	{ "device", "query",			"None",						0, TestAudioDeviceQuery				}, // AUDIO_TEST_DEVICE_QUERY
+	{ "device", "out",				"None",						0, TestAudioDeviceOutputSimple		}, // AUDIO_TEST_DEVICE_OUTPUT_SIMPLE
+	{ "device", "out_fm",			"None",						0, TestAudioDeviceOutputFm			}, // AUDIO_TEST_DEVICE_OUTPUT_FM
+	{ "device", "out_pan",			"None",						0, TestAudioDeviceOutputPan			}, // AUDIO_TEST_DEVICE_OUTPUT_PAN
+	{ "source", "convert",			"SourcePath",				1, TestAudioSourceConvert			}, // AUDIO_TEST_SOURCE_CONVERT
+	{ "system", "emitter_manager",	"None",						0, TestAudioSystemEmitterManager	}, // AUDIO_TEST_SYSTEM_EMITTER_MANAGER
+	{ "system", "voice_manager",	"SourcePath or Directory",	1, TestAudioSystemVoiceManager		}, // AUDIO_TEST_SYSTEM_VOICE_MANAGER
+	{ "system", "soundfile_manager", "None",					1, TestAudioSystemSoundFileManager  }, // AUDIO_TEST_SYSTEM_SOUND_MANAGER
 };
 static_assert(ARRAY_COUNT(AudioTestInfoList) == AUDIO_TESTS, "Mismatched info list and test enumeration");
 
diff --git a/Engine/Source/Runtime/UnrealAudio/Private/Tests/UnrealAudioDeviceTests.cpp b/Engine/Source/Runtime/UnrealAudio/Private/Tests/UnrealAudioDeviceTests.cpp
--- a/Engine/Source/Runtime/UnrealAudio/Private/Tests/UnrealAudioDeviceTests.cpp
+++ b/Engine/Source/Runtime/UnrealAudio/Private/Tests/UnrealAudioDeviceTests.cpp
@@ -11,22 +11,18 @@
 
 namespace UAudio
 {
-	static FUnrealAudioModule* UnrealAudioModule = nullptr;
-	static bool bTestInitialized = false;
+	/** Static singleton test generator object. */
+	static IUnrealAudioModule* UnrealAudioModule = nullptr;
+	static Test::IGenerator* TestGenerator = nullptr;
+	static FThreadSafeBool bTestActive = false;
 
-	void FUnrealAudioModule::InitializeTests(FUnrealAudioModule* Module)
+	void FUnrealAudioModule::InitializeDeviceTests(IUnrealAudioModule* Module)
 	{
 		UnrealAudioModule = Module;
-		bTestInitialized = true;
 	}
 
 	bool TestDeviceQuery()
 	{
-		if (!bTestInitialized)
-		{
-			return false;
-		}
-
 		UE_LOG(LogUnrealAudio, Display, TEXT("Running audio device query test..."));
 
 		IUnrealAudioDeviceModule* UnrealAudioDevice = UnrealAudioModule->GetDeviceModule();
@@ -92,13 +88,15 @@ namespace UAudio
 		return true;
 	}
 
-	/**
-	* TestCallback
-	* Callback function used for all device test functions. 
-	* Function takes input params and passes to the ITestGenerator object to generate output buffers.
-	*/
-	static bool TestCallback(FCallbackInfo& CallbackInfo)	
+	bool FUnrealAudioModule::DeviceTestCallback(struct FCallbackInfo& CallbackInfo)
 	{
+		if (!bTestActive)
+		{
+			return true;
+		}
+
+		check(TestGenerator);
+
 		static int32 CurrentSecond = -1;
 		int32 StreamSecond = (int32)CallbackInfo.StreamTime;
 		if (StreamSecond != CurrentSecond)
@@ -114,22 +112,14 @@ namespace UAudio
 		// Sets any data used by lots of different generators in static data struct
 		Test::UpdateCallbackData(CallbackInfo);
 
-		// Cast the user data object to our own FSimpleFM object (since we know we used that when we created the stream)
-		Test::IGenerator* Generator = (Test::IGenerator*)CallbackInfo.UserData;
-
-		// Get the next buffer of output data
-		Generator->GetNextBuffer(CallbackInfo);
-		return true;
+		return TestGenerator->GetNextBuffer(CallbackInfo);
 	}
 
 	static bool DoOutputTest(const FString& TestName, double LifeTime, Test::IGenerator* Generator)
 	{
 		check(Generator);
-
-		if (!bTestInitialized)
-		{
-			return false;
-		}
+		check(TestGenerator == nullptr);
+		check(bTestActive == false);
 
 		UE_LOG(LogUnrealAudio, Display, TEXT("Running audio test: \"%s\"..."), *TestName);
 
@@ -142,69 +132,17 @@ namespace UAudio
 			UE_LOG(LogUnrealAudio, Display, TEXT("Time of test: (indefinite)"));
 		}
 
-		IUnrealAudioDeviceModule* UnrealAudioDevice = UnrealAudioModule->GetDeviceModule();
-
-		EDeviceApi::Type DeviceApi;
-		if (!UnrealAudioDevice->GetDevicePlatformApi(DeviceApi))
-		{
-			UE_LOG(LogUnrealAudio, Display, TEXT("Failed."));
-			return false;
-		}
-		if (DeviceApi == EDeviceApi::DUMMY)
-		{
-			UE_LOG(LogUnrealAudio, Display, TEXT("This is a dummy API. Platform not implemented yet."));
-			return true;
-		}
-
-		uint32 DefaultDeviceIndex = 0;
-		if (!UnrealAudioDevice->GetDefaultOutputDeviceIndex(DefaultDeviceIndex))
-		{
-			UE_LOG(LogUnrealAudio, Display, TEXT("Failed to get default device index."));
-			return false;
-		}
-
-		FDeviceInfo DeviceInfo;
-		if (!UnrealAudioDevice->GetOutputDeviceInfo(DefaultDeviceIndex, DeviceInfo))
-		{
-			UE_LOG(LogUnrealAudio, Display, TEXT("Failed to get default device info."));
-			return false;
-		}
-
-		FCreateStreamParams CreateStreamParams;
-		CreateStreamParams.OutputDeviceIndex = DefaultDeviceIndex;
-		CreateStreamParams.CallbackFunction = &TestCallback;
-		CreateStreamParams.UserData = (void*)Generator;
-		CreateStreamParams.CallbackBlockSize = 1024;
-
-		UE_LOG(LogUnrealAudio, Display, TEXT("Creating the stream..."));
-		if (!UnrealAudioDevice->CreateStream(CreateStreamParams))
-		{
-			UE_LOG(LogUnrealAudio, Display, TEXT("Failed to create a stream."));
-			return false;
-		}
-
-		UE_LOG(LogUnrealAudio, Display, TEXT("Starting the stream."));
-		UnrealAudioDevice->StartStream();
+		TestGenerator = Generator;
+		bTestActive = true;
 
-		// Block until the synthesizer is done
+		// Block this thread until the synthesizer is done
 		while (!Generator->IsDone())
 		{
 			FPlatformProcess::Sleep(1);
 		}
 
-		// Stop the output stream (which blocks until its actually freed)
-		UE_LOG(LogUnrealAudio, Display, TEXT("Shutting down the stream."));
-		if (!UnrealAudioDevice->StopStream())
-		{
-			UE_LOG(LogUnrealAudio, Display, TEXT("Failed to stop the device stream."));
-			return false;
-		}
-
-		if (!UnrealAudioDevice->ShutdownStream())
-		{
-			UE_LOG(LogUnrealAudio, Display, TEXT("Failed to shut down the device stream."));
-			return false;
-		}
+		bTestActive = false;
+		TestGenerator = nullptr;
 
 		// At this point audio device I/O is done
 		UE_LOG(LogUnrealAudio, Display, TEXT("Success!"));
@@ -229,43 +167,18 @@ namespace UAudio
 		return DoOutputTest("output white noise panner", LifeTime, &SimpleWhiteNoisePan);
 	}
 
-	bool TestSourceImport(const FSoundFileImportSettings& ImportSettings)
+	bool TestSourceConvert(const FString& FilePath, const FSoundFileConvertFormat& ConvertFormat)
 	{
-		UE_LOG(LogUnrealAudio, Display, TEXT("Testing exporting a single sound source."));
-
-		UE_LOG(LogUnrealAudio, Display, TEXT("First importing."));
-		UE_LOG(LogUnrealAudio, Display, TEXT("Import Path: %s"), *ImportSettings.SoundFilePath);
-		UE_LOG(LogUnrealAudio, Display, TEXT("Import Format: %s"), *FString::Printf(TEXT("%s - %s"), ESoundFileFormat::ToStringMajor(ImportSettings.Format), ESoundFileFormat::ToStringMinor(ImportSettings.Format)));
-		UE_LOG(LogUnrealAudio, Display, TEXT("Import SampleRate: %d"), ImportSettings.SampleRate);
-		UE_LOG(LogUnrealAudio, Display, TEXT("Import EncodingQuality: %.2f"), ImportSettings.EncodingQuality);
-		UE_LOG(LogUnrealAudio, Display, TEXT("Perform Peak Normalization: %s"), ImportSettings.bPerformPeakNormalization ? TEXT("Yes") : TEXT("No"));
-
-		
-		TSharedPtr<ISoundFile> ImportedSoundFile = UnrealAudioModule->ImportSound(ImportSettings);
-		ESoundFileState::Type State = ImportedSoundFile->GetState();
-		while (State != ESoundFileState::IMPORTED)
+		//  Check if the path is a folder or a single file
+		if (FPaths::DirectoryExists(FilePath))
 		{
-			FPlatformProcess::Sleep(1);
-
-			State = ImportedSoundFile->GetState();
-			if (State == ESoundFileState::HAS_ERROR)
-			{
-				UE_LOG(LogUnrealAudio, Error, TEXT("Failed to import sound source!"));
-				UE_LOG(LogUnrealAudio, Error, TEXT("Error: %s"), ESoundFileError::ToString(ImportedSoundFile->GetError()));
-				return false;
-			}
-		}
-		UE_LOG(LogUnrealAudio, Display, TEXT("Succeeded importing sound source!"));
-		return true;
-	}
+			// Test the first file
+			TArray<FString> SoundFiles;
+			FString Directory = FilePath;
+			GetSoundFileListInDirectory(Directory, SoundFiles);
 
-	bool TestSourceImportExport(const FSoundFileImportSettings& ImportSettings)
-	{
-		//  Check if the soundfile path is a folder or a single file:
-		if (FPaths::DirectoryExists(ImportSettings.SoundFilePath))
-		{
-			FString FolderPath = ImportSettings.SoundFilePath;
-			UE_LOG(LogUnrealAudio, Display, TEXT("Testing import then export of all audio files in directory: %s."), *FolderPath);
+			FString FolderPath = FilePath;
+			UE_LOG(LogUnrealAudio, Display, TEXT("Testing import then export of all audio files in directory: %s."), *FilePath);
 
 			// Create the output exported directory if it doesn't exist
 			FString OutputDir = FolderPath / "Exported";
@@ -274,138 +187,44 @@ namespace UAudio
 				FPlatformFileManager::Get().GetPlatformFile().CreateDirectory(*OutputDir);
 			}
 
-			// Get list of .wav files:
-			TArray<FString> InputFiles;
-
-			// Get input files in formats we're currently supporting
-			IFileManager::Get().FindFiles(InputFiles, *(FolderPath / "*.wav"), true, false);
-			IFileManager::Get().FindFiles(InputFiles, *(FolderPath / "*.aif"), true, false);
-			IFileManager::Get().FindFiles(InputFiles, *(FolderPath / "*.flac"), true, false);
-			IFileManager::Get().FindFiles(InputFiles, *(FolderPath / "*.ogg"), true, false);
-
-			// Make an array of shared ptr sound files
-			TArray<TSharedPtr<ISoundFile>> SoundFiles;
-
-			// Array of sound files (names) that have finished importing
-			TSet<FString> FinishedImporting;
-
-			// Array of sound files (names) that have error'd out during import
-			TSet<FString> HasError;
-
-			FSoundFileImportSettings Settings = ImportSettings;
-
-			// Import all the files!
-			UE_LOG(LogUnrealAudio, Display, TEXT("Importing..."));
-			for (FString& InputFile : InputFiles)
+			FString SoundFileExtension;
+			if (!GetFileExtensionForFormatFlags(ConvertFormat.Format, SoundFileExtension))
 			{
-				Settings.SoundFilePath = FolderPath / InputFile;
-				UE_LOG(LogUnrealAudio, Display, TEXT("%s"), *InputFile);
-
-				// start the import process and store the ISoundFile ptr in our array of in-flight imports
-				TSharedPtr<ISoundFile> ImportedSoundFile = UnrealAudioModule->ImportSound(Settings);
-				SoundFiles.Add(ImportedSoundFile);
+				UE_LOG(LogUnrealAudio, Error, TEXT("Unknown sound file format."));
+				return false;
 			}
 
-			// Loop through in-flight imports and start exports as soon as the file finishes importing
-			bool bIsFinishedImporting = false;
-			while (!bIsFinishedImporting)
+			// Convert and export all the files!
+			UE_LOG(LogUnrealAudio, Display, TEXT("Converting and exporting..."));
+			for (FString& InputFile : SoundFiles)
 			{
-				// Assume we're done, but set false as soon as we find an import still going
-				bIsFinishedImporting = true;
+				UE_LOG(LogUnrealAudio, Display, TEXT("%s"), *InputFile);
 
-				// Loop through our stored sound files
-				for (TSharedPtr<ISoundFile> SoundFile : SoundFiles)
+				// We're building up a filename with [NAME]_exported.[EXT] in the output directory
+				FString BaseSoundFileName = FPaths::GetBaseFilename(InputFile);
+				FString OutputFileName = BaseSoundFileName + TEXT("_exported.") + SoundFileExtension;
+				FString OutputPath = OutputDir / OutputFileName;
+
+				while (UnrealAudioModule->GetNumBackgroundTasks() > 2)
 				{
-					// Get the sound file state, handle any errors that occurred during import
-					// Note: its ok for our test assets here to fail since we're testing graceful failure. We
-					// might have some pathological sound assets in the test asset directory.
-					ESoundFileState::Type State = SoundFile->GetState();
-					if (State == ESoundFileState::HAS_ERROR)
-					{
-						FString Name;
-						SoundFile->GetName(Name);
-
-						// Store the error in a list of error'd sound files
-						if (!HasError.Contains(Name))
-						{
-							HasError.Add(Name);
-						}
-					}
-					else if (State != ESoundFileState::IMPORTED)
-					{
-						bIsFinishedImporting = false;
-					}
-					else
-					{
-						FString Name;
-						SoundFile->GetName(Name);
-
-						// If this file's state is IMPORTED, then we add it to the list
-						// of FinishedImporting sounds, then trigger the export process.
-						if (!FinishedImporting.Contains(Name))
-						{
-							FinishedImporting.Add(Name);
-
-							UE_LOG(LogUnrealAudio, Display, TEXT("Finished importing %s"), *Name);
-
-							TSharedPtr<ISoundFileData> SoundFileData;
-							SoundFile->GetSoundFileData(SoundFileData);
-
-							// We're building up a filename with [NAME]_exported.[EXT] in the output directory
-							const FString& SoundFileFullPath = SoundFileData->GetFilePath();
-							FString BaseSoundFileName = FPaths::GetBaseFilename(SoundFileFullPath);
-							FString SoundFileExtension = FPaths::GetExtension(SoundFileFullPath);
-							FString SoundFilePath = FPaths::GetPath(SoundFileFullPath);
-							FString OutputFileName = BaseSoundFileName + TEXT("_exported.") + SoundFileExtension;
-							FString OutputPath = OutputDir / OutputFileName;
-
-							UE_LOG(LogUnrealAudio, Display, TEXT("Now exporting to %s"), *OutputPath);
-
-							UnrealAudioModule->ExportSound(SoundFileData, OutputPath);
-						}
-					}
+					FPlatformProcess::Sleep(0.001f);
 				}
-
-				// Sleep until all imports have finished. In-flight exports
-				// will cause the UnrealAudio shutdown to block until they are finished.
-				FPlatformProcess::Sleep(1);
+				UnrealAudioModule->ConvertSound(InputFile, OutputPath, ConvertFormat);
 			}
-
 			return true;
 		}
-		else if (FPaths::FileExists(ImportSettings.SoundFilePath))
+		else if (FPaths::FileExists(FilePath))
 		{
-			UE_LOG(LogUnrealAudio, Display, TEXT("Testing import, then export of a single sound source: %s."), *ImportSettings.SoundFilePath);
-			UE_LOG(LogUnrealAudio, Display, TEXT("Import Format: %s"), *FString::Printf(TEXT("%s - %s"), ESoundFileFormat::ToStringMajor(ImportSettings.Format), ESoundFileFormat::ToStringMinor(ImportSettings.Format)));
-			UE_LOG(LogUnrealAudio, Display, TEXT("Import SampleRate: %d"), ImportSettings.SampleRate);
-			UE_LOG(LogUnrealAudio, Display, TEXT("Import EncodingQuality: %.2f"), ImportSettings.EncodingQuality);
-			UE_LOG(LogUnrealAudio, Display, TEXT("Perform Peak Normalization: %s"), ImportSettings.bPerformPeakNormalization ? TEXT("Yes") : TEXT("No"));
-
-			// Start the import process on the sound
-			TSharedPtr<ISoundFile> ImportedSoundFile = UnrealAudioModule->ImportSound(ImportSettings);
-
-			// Wait for the import to finish (block this thread for test-purposes). Normally, this would be non-blocking.
-			ESoundFileState::Type State = ImportedSoundFile->GetState();
-			while (State != ESoundFileState::IMPORTED)
-			{
-				FPlatformProcess::Sleep(1);
-
-				State = ImportedSoundFile->GetState();
-				if (State == ESoundFileState::HAS_ERROR)
-				{
-					UE_LOG(LogUnrealAudio, Error, TEXT("Failed to import sound source!"));
-					UE_LOG(LogUnrealAudio, Error, TEXT("Error: %s"), ESoundFileError::ToString(ImportedSoundFile->GetError()));
-					return false;
-				}
-			}
-
-			UE_LOG(LogUnrealAudio, Display, TEXT("Succeeded importing sound source!"));
+			UE_LOG(LogUnrealAudio, Display, TEXT("Testing import, then export of a single sound source: %s."), *FilePath);
+			UE_LOG(LogUnrealAudio, Display, TEXT("Convert Format: %s"), *FString::Printf(TEXT("%s - %s"), ESoundFileFormat::ToStringMajor(ConvertFormat.Format), ESoundFileFormat::ToStringMinor(ConvertFormat.Format)));
+			UE_LOG(LogUnrealAudio, Display, TEXT("Convert SampleRate: %d"), ConvertFormat.SampleRate);
+			UE_LOG(LogUnrealAudio, Display, TEXT("Convert EncodingQuality: %.2f"), ConvertFormat.EncodingQuality);
+			UE_LOG(LogUnrealAudio, Display, TEXT("Perform Peak Normalization: %s"), ConvertFormat.bPerformPeakNormalization ? TEXT("Yes") : TEXT("No"));
 
 			// Now setup the export path
-			FString SoundFileFullPath = ImportSettings.SoundFilePath;
-			FString BaseSoundFileName = FPaths::GetBaseFilename(SoundFileFullPath);
-			FString SoundFileExtension = FPaths::GetExtension(SoundFileFullPath);
-			FString SoundFileDir = FPaths::GetPath(SoundFileFullPath);
+			FString BaseSoundFileName = FPaths::GetBaseFilename(FilePath);
+			FString SoundFileExtension = FPaths::GetExtension(FilePath);
+			FString SoundFileDir = FPaths::GetPath(FilePath);
 
 			// Create the export directory if it doesn't exist
 			FString OutputDir = SoundFileDir;
@@ -417,15 +236,12 @@ namespace UAudio
 			// Append _exported to the file path just to make it clear that this is the exported version of the file
 			FString OutputPath = SoundFileDir / (BaseSoundFileName + TEXT("_exported.") + SoundFileExtension);
 
-			// Start the export process. We can exit this test function since UnrealAudio will block on shutdown
-			// until all background processes have finished. Normally the Editor will remain running during the export process.
-			TSharedPtr<ISoundFileData> SoundFileData;
-			ImportedSoundFile->GetSoundFileData(SoundFileData);
-			UnrealAudioModule->ExportSound(SoundFileData, OutputPath);
+			UnrealAudioModule->ConvertSound(FilePath, OutputPath, ConvertFormat);
+
 			return true;
 		}
 
-		UE_LOG(LogUnrealAudio, Display, TEXT("Path %s is not a single file or a directory."), *ImportSettings.SoundFilePath);
+		UE_LOG(LogUnrealAudio, Display, TEXT("Path %s is not a single file or a directory."), *FilePath);
 		return false;
 	}
 
diff --git a/Engine/Source/Runtime/UnrealAudio/Private/Tests/UnrealAudioSystemTests.cpp b/Engine/Source/Runtime/UnrealAudio/Private/Tests/UnrealAudioSystemTests.cpp
new file mode 100644
--- /dev/null
+++ b/Engine/Source/Runtime/UnrealAudio/Private/Tests/UnrealAudioSystemTests.cpp
@@ -0,0 +1,381 @@
+// Copyright 1998-2015 Epic Games, Inc. All Rights Reserved.
+
+#include "UnrealAudioPrivate.h"
+#include "UnrealAudioModule.h"
+#include "UnrealAudioEmitterInternal.h"
+
+#if ENABLE_UNREAL_AUDIO
+
+namespace UAudio
+{
+	static IUnrealAudioModule* UnrealAudioModuleSystem = nullptr;
+	static bool bTestInitialized = false;
+
+	void FUnrealAudioModule::InitializeSystemTests(IUnrealAudioModule* Module)
+	{
+		UnrealAudioModuleSystem = Module;
+		bTestInitialized = true;
+	}
+
+	// this is called from the system thread
+	void FUnrealAudioModule::UpdateSystemTests()
+	{
+
+	}
+
+	class FTestTimer
+	{
+	public:
+		FTestTimer()
+			: TotalTime(0.0)
+			, StartTime(0.0)
+		{
+		}
+
+		FTestTimer(double InTotalTimeSeconds)
+			: TotalTime(InTotalTimeSeconds)
+			, StartTime(FPlatformTime::Seconds())
+		{
+		}
+
+		void Start(double Time)
+		{
+			StartTime = FPlatformTime::Seconds();
+			TotalTime = Time;
+		}
+
+		double GetTimeFraction()
+		{
+			return FMath::Min((FPlatformTime::Seconds() - StartTime) / TotalTime, 1.0);
+		}
+
+		bool IsDone()
+		{
+			return GetTimeFraction() >= 1.0;
+		}
+
+	private:
+		double TotalTime;
+		double StartTime;
+	};
+
+
+	static FVector GetRandomPosition()
+	{
+		return FVector(FMath::FRandRange(-100.0f, 100.0f), FMath::FRandRange(-100.0f, 100.0f), FMath::FRandRange(-100.0f, 100.0f));
+	}
+
+
+	bool TestEmitterManager()
+	{
+		if (!bTestInitialized)
+		{
+			return false;
+		}
+		
+		// Lets create a bunch of initial emitters
+		TLinkedList<TSharedPtr<IEmitter>>* EmittersHead = nullptr;
+		int32 NumEmitters = 100;
+		for (int32 i = 0; i < NumEmitters; ++i)
+		{
+			TSharedPtr<IEmitter> Emitter = UnrealAudioModuleSystem->EmitterCreate();
+			UE_LOG(LogUnrealAudio, Display, TEXT("[%d] Create"), Emitter->GetId());
+
+			Emitter->SetPosition(GetRandomPosition());
+			TLinkedList<TSharedPtr<IEmitter>>* EmitterLink = new TLinkedList<TSharedPtr<IEmitter>>(Emitter);
+			if (EmittersHead == nullptr)
+			{
+				EmittersHead = EmitterLink;
+			}
+			else
+			{
+				EmitterLink->Link(EmittersHead);
+			}
+		}
+
+		// For 30 seconds, lets randomly set emitter positions from this thread
+		FTestTimer TestTimer(30);
+		while (!TestTimer.IsDone())
+		{
+			TLinkedList<TSharedPtr<IEmitter>>* Link = EmittersHead;
+			while (Link)
+			{
+				// 50-50 chance to change position
+				if (FMath::FRand() < 0.5f)
+				{
+					FVector Pos = GetRandomPosition();
+					UE_LOG(LogUnrealAudio, Display, TEXT("[%d] SetPosition: (%.2f, %.2f, %.2f)"), (*Link)->Get()->GetId(), Pos.X, Pos.Y, Pos.Z);
+					(*Link)->Get()->SetPosition(Pos);
+				}
+				// 5% chance to be released
+				else if (FMath::FRand() < 0.05f)
+				{
+					UE_LOG(LogUnrealAudio, Display, TEXT("[%d] Release"), (*Link)->Get()->GetId());
+					(*Link)->Get()->Release();
+
+					// Unlink the emitter from our linked list and delete it
+					TLinkedList<TSharedPtr<IEmitter>>* NextLink = Link->GetNextLink();
+					Link->Unlink();
+					delete Link;
+					Link = NextLink;
+				}
+				else
+				{
+					Link = Link->GetNextLink();
+				}
+			}
+
+			// Randomly add a bunch of emitters
+			for (int32 i = 0; i < 10; ++i)
+			{
+				if (FMath::FRand() < 0.1f)
+				{
+					TSharedPtr<IEmitter> Emitter = UnrealAudioModuleSystem->EmitterCreate();
+					UE_LOG(LogUnrealAudio, Display, TEXT("[%d] Create"), Emitter->GetId());
+					Emitter->SetPosition(GetRandomPosition());
+					TLinkedList<TSharedPtr<IEmitter>>* EmitterLink = new TLinkedList<TSharedPtr<IEmitter>>(Emitter);
+					EmitterLink->Link(EmittersHead);
+				}
+			}
+
+			FPlatformProcess::Sleep(0.033f);
+		}
+
+		return true;
+	}
+
+	bool TestVoiceManager(const FString& FolderOrPath)
+	{
+		if (!bTestInitialized)
+		{
+			return false;
+		}
+
+		TSharedPtr<ISoundFile> SoundFile = UnrealAudioModuleSystem->StreamSoundFile(*FolderOrPath);
+
+		FVoiceInitializationParams InitParams;
+		InitParams.SoundFile = SoundFile;
+		InitParams.bIsLooping = true;
+		InitParams.BaselinePitchScale = 1.0f;
+		InitParams.BaselineVolumeScale = 1.0f;
+		InitParams.PriorityWeight = 1.0f;
+
+		TSharedPtr<IVoice> Voice = UnrealAudioModuleSystem->VoiceCreate(InitParams);
+		EVoiceError::Type Error = Voice->Play();
+		check(Error == EVoiceError::NONE);
+		
+		FTestTimer VolumeTimer;
+		FTestTimer PitchTimer;
+		float VolumeProduct = 0.0f;
+		float PitchProduct = 0.0f;
+		float VolumeTime = 0.0f;
+		float VolumeTarget = 0.0f;
+		float PitchTime = 0.0f;
+		float PitchTarget = 0.0f;
+
+		while (true)
+		{
+			// Update the module on the main thread
+			UnrealAudioModuleSystem->Update();
+
+			if (VolumeTimer.IsDone())
+			{
+				VolumeTime = FMath::FRandRange(1.0f, 10.0f);
+				VolumeTarget = FMath::FRandRange(0.0f, 1.0f);
+
+				Error = Voice->SetVolumeScale(VolumeTarget, VolumeTime);
+				check(Error == EVoiceError::NONE);
+
+				VolumeTimer.Start(VolumeTime);
+			}
+
+			if (PitchTimer.IsDone())
+			{
+				PitchTime = FMath::FRandRange(1.0f, 10.0f);
+				PitchTarget = FMath::FRandRange(0.01f, 4.0f);
+
+				Error = Voice->SetPitchScale(PitchTarget, PitchTime);
+				check(Error == EVoiceError::NONE);
+
+				PitchTimer.Start(PitchTime);
+			}
+
+ 			Voice->GetVolumeProduct(VolumeProduct);
+ 			Voice->GetPitchProduct(PitchProduct);
+
+			float CurrentVolumeTime = VolumeTime * VolumeTimer.GetTimeFraction();
+			float CurrentPitchTime = PitchTime * PitchTimer.GetTimeFraction();
+
+			UE_LOG(LogUnrealAudio, Display, TEXT("V: %.2f [%.2f, %.2f/%.2f] P: %.2f [%.2f, %.2f/%.2f]"), 
+				VolumeProduct, 
+				VolumeTarget,
+				CurrentVolumeTime,
+				VolumeTime,
+				
+				PitchProduct,
+				PitchTarget,
+				CurrentPitchTime,
+				PitchTime
+			);
+
+			FPlatformProcess::Sleep(0.033f);
+		}
+
+		return true;
+	}
+
+	static bool TestSingleFileLoad(const FString& Path)
+	{
+		{
+			// Test loading a single file
+			UE_LOG(LogUnrealAudio, Display, TEXT("Loading sound file %s"), *Path);
+			TSharedPtr<ISoundFile> SoundFile = UnrealAudioModuleSystem->LoadSoundFile(*Path);
+
+			// Block the thread till the sound file is loaded
+			ESoundFileState::Type State;
+			SoundFile->GetState(State);
+			while (State != ESoundFileState::LOADED)
+			{
+				FPlatformProcess::Sleep(0.033f);
+				UE_LOG(LogUnrealAudio, Display, TEXT("Loading..."));
+				SoundFile->GetState(State);
+			}
+
+			UE_LOG(LogUnrealAudio, Display, TEXT("Loaded."));
+
+			// Try loading it again
+			UE_LOG(LogUnrealAudio, Display, TEXT("Attempting to reload sound file from cache to new ISoundFile ptr."));
+			TSharedPtr<ISoundFile> NewSoundFile = UnrealAudioModuleSystem->LoadSoundFile(*Path);
+
+			// Since it's already loaded, this should immediately be loaded
+			NewSoundFile->GetState(State);
+			check(State == ESoundFileState::LOADED);
+
+			UE_LOG(LogUnrealAudio, Display, TEXT("Sound file already loaded: Yes"));
+
+			// Let the sound file shared ptrs fall out of scope, this will cause the sound files
+			// to decref the sound file manager internally and trigger the time-based flush
+		}
+
+		UE_LOG(LogUnrealAudio, Display, TEXT("Waiting until the sound file is flushed from the sound file cache...."));
+
+		// Wait now until the loaded file is flushed from the cache
+		float TotalTime = 0.0f;
+		int32 NumSoundFilesLoaded = 1;
+		while (NumSoundFilesLoaded != 0)
+		{
+			UnrealAudioModuleSystem->Update();
+			NumSoundFilesLoaded = UnrealAudioModuleSystem->GetNumSoundFilesLoaded();
+
+			UE_LOG(LogUnrealAudio, Display, TEXT("Num Loaded: %d (%.2f seconds)"), NumSoundFilesLoaded, TotalTime);
+
+			FPlatformProcess::Sleep(0.033f);
+			TotalTime += 0.033f;
+		}
+
+		UE_LOG(LogUnrealAudio, Display, TEXT("Sound file memory cache has flushed in %.2f seconds"), TotalTime);
+		return true;
+	}
+
+	static bool TestArrayOfFilesLoading(TArray<FString>& Files)
+	{
+		TArray<TSharedPtr<ISoundFile>> LoadedFiles;
+
+		float CurrentTime = 0.0;
+		float TimeSinceLastPrint = 0.0f;
+
+		while (CurrentTime < 30.0f)
+		{
+			UnrealAudioModuleSystem->Update();
+
+			// Randomly load a new file
+			if (FMath::FRandRange(0.0f, 100.0f) < 10.0f)
+			{
+				// Randomly pick a file from the list of input files
+				int Index = FMath::RandRange(0, Files.Num() - 1);
+				FString& Path = Files[Index];
+
+				// Load that baby up
+				TSharedPtr<ISoundFile> SoundFile;
+				if (FMath::RandBool())
+				{
+					SoundFile = UnrealAudioModuleSystem->LoadSoundFile(*Path);
+				}
+				else
+				{
+					SoundFile = UnrealAudioModuleSystem->StreamSoundFile(*Path);
+				}
+
+				// Try to store the new loaded file in a previously used slot..
+				bool bFoundEmptySlot = false;
+				for (int32 i = 0; i < LoadedFiles.Num(); ++i)
+				{
+					if (!LoadedFiles[i].IsValid())
+					{
+						LoadedFiles[i] = SoundFile;
+						bFoundEmptySlot = true;
+						break;
+					}
+				}
+
+				// Otherwise add it to our list of loaded files
+				if (!bFoundEmptySlot)
+				{
+					LoadedFiles.Add(SoundFile);
+				}
+			}
+
+			// Loop through the loaded files and randomly unref them
+			// which will cause them to become "inactive" if nobody else is refing them.
+			// This will give the sound file manager a chance to flush the asset over time
+			// or if the memory threshold is reached.
+			for (int i = 0; i < LoadedFiles.Num(); ++i)
+			{
+				if (LoadedFiles[i].IsValid() && FMath::FRandRange(0.0f, 100.0f) < 5.0f)
+				{
+					LoadedFiles[i] = nullptr;
+				}
+			}
+
+			if (CurrentTime - TimeSinceLastPrint >= 1.0f)
+			{
+				UnrealAudioModuleSystem->LogSoundFileMemoryInfo();
+				TimeSinceLastPrint = CurrentTime;
+			}
+
+			CurrentTime += 0.033f;
+			FPlatformProcess::Sleep(0.033f);
+		}
+
+		return true;
+	}
+
+
+	bool TestSoundFileManager(const FString& FolderOrPath)
+	{
+		if (FPaths::DirectoryExists(FolderOrPath))
+		{
+			// Test the first file
+			TArray<FString> SoundFiles;
+			FString Directory = FolderOrPath;
+			GetSoundFileListInDirectory(Directory, SoundFiles);
+
+			if (SoundFiles.Num() > 0)
+			{
+				return TestArrayOfFilesLoading(SoundFiles);
+			}
+			else
+			{
+				UE_LOG(LogUnrealAudio, Error, TEXT("Failed to find any sound files at directory %s"), *FolderOrPath);
+				return false;
+			}
+		}
+		else
+		{
+			return TestSingleFileLoad(FolderOrPath);
+		}
+	}
+
+}
+
+#endif // #if ENABLE_UNREAL_AUDIO
diff --git a/Engine/Source/Runtime/UnrealAudio/Private/UnrealAudio.cpp b/Engine/Source/Runtime/UnrealAudio/Private/UnrealAudio.cpp
--- a/Engine/Source/Runtime/UnrealAudio/Private/UnrealAudio.cpp
+++ b/Engine/Source/Runtime/UnrealAudio/Private/UnrealAudio.cpp
@@ -19,9 +19,37 @@ IMPLEMENT_MODULE(UAudio::FUnrealAudioModule, UnrealAudio);
 
 namespace UAudio
 {
+	/**
+	 * Audio device callback function, wraps UnrealAudio module and fowards call.
+	 *
+	 * @param [in,out]	CallbackInfo	Information describing the audio device callback.
+	 *
+	 * @return	true if it succeeds, false if it fails.
+	 */
+
+	static bool AudioDeviceCallbackFunc(FCallbackInfo& CallbackInfo)
+	{
+		FUnrealAudioModule* AudioModule = (FUnrealAudioModule*)CallbackInfo.UserData;
+		return AudioModule->AudioDeviceCallback(CallbackInfo);
+	}
+
+	/************************************************************************/
+	/* FUnrealAudioModule Implementation									*/
+	/************************************************************************/
+
 	FUnrealAudioModule::FUnrealAudioModule()
 		: UnrealAudioDevice(nullptr)
+		, NumBackgroundTasks(0)
+		, EmitterManager(this)
+		, VoiceManager(this)
+		, SoundFileManager(this)
 		, SoundFileDllHandle(nullptr)
+		, bIsStoppingSystemThread(false)
+		, AudioSystemTimeSec(0.0)
+		, SystemThread(nullptr)
+		, AudioThreadCommandQueue(500)
+		, MainThreadCommandQueue(100)
+		, SystemThreadUpdateTime(0.033)
 	{
 	}
 
@@ -35,6 +63,17 @@ namespace UAudio
 		return InitializeInternal();
 	}
 
+	void FUnrealAudioModule::Update()
+	{
+		MainThreadChecker.CheckThread();
+
+		ExecuteMainThreadCommands();
+
+		// Update the sound file manager from the main thread.
+		SoundFileManager.Update();
+
+	}
+
 	bool FUnrealAudioModule::Initialize(const FString& DeviceModuleName)
 	{
 		FString Name("UnrealAudio");
@@ -46,32 +85,130 @@ namespace UAudio
 	bool FUnrealAudioModule::InitializeInternal()
 	{
 		bool bSuccess = false;
+
+		bSuccess = InitializeAudioDevice();
+		check(bSuccess);
+
+		bSuccess = InitializeAudioSystem();
+		check(bSuccess);
+
+		return bSuccess;
+	}
+
+	bool FUnrealAudioModule::InitializeAudioSystem()
+	{
+		bool bSuccess = false;
+
+		FVoiceManagerSettings VoiceManagerSettings;
+		VoiceManagerSettings.MaxVoiceCount = 32;
+		VoiceManagerSettings.MaxVirtualVoiceCount = 1000;
+		VoiceManagerSettings.MaxPitch = 4.0f;
+		VoiceManagerSettings.MinPitch = 0.01f;
+		VoiceManagerSettings.ControlUpdateRateSeconds = SystemThreadUpdateTime;
+		VoiceManagerSettings.NumDecoders = 2;
+		VoiceManagerSettings.DecoderSettings.DecodeBufferFrames = 1024;
+		VoiceManagerSettings.DecoderSettings.NumDecodeBuffers = 3;
+
+		VoiceManager.Init(VoiceManagerSettings);
+
+		FSoundFileManagerSettings SoundFileManagerSettings;
+		SoundFileManagerSettings.MaxNumberOfLoadedSounds = 5000;
+		SoundFileManagerSettings.TargetMemoryLimit = 10 * 1024 * 1024; // Convert MB to Bytes
+		SoundFileManagerSettings.NumLoadingThreads = 2;
+		SoundFileManagerSettings.FlushTimeThreshold = 10.0f;
+		SoundFileManagerSettings.TimeDeltaPerUpdate = 0.033f;
+		SoundFileManagerSettings.LoadingThreadPriority = EThreadPriority::TPri_Normal;
+
+		SoundFileManager.Init(SoundFileManagerSettings);
+
+		// Define the default sound file load settings
+		DefaultConvertFormat.bPerformPeakNormalization = false;
+		DefaultConvertFormat.EncodingQuality = 0.75;
+		DefaultConvertFormat.Format = UAudio::ESoundFileFormat::OGG | UAudio::ESoundFileFormat::VORBIS;
+		DefaultConvertFormat.SampleRate = 44100;
+
+		bSuccess = LoadSoundFileLib();
+
+		if (bSuccess)
+		{
+			// Store this thread ID as the "main" thread id
+			MainThreadChecker.InitThread();
+
+			// Create the system thread
+			bIsStoppingSystemThread = false;
+			SystemThread = FRunnableThread::Create(this, TEXT("Audio System Thread"), 0, TPri_Normal);
+		}
+
+		InitializeSystemTests(this);
+
+		return bSuccess;
+	}
+
+	bool FUnrealAudioModule::InitializeAudioDevice()
+	{
 		UnrealAudioDevice = FModuleManager::LoadModulePtr<IUnrealAudioDeviceModule>(ModuleName);
 		if (UnrealAudioDevice)
 		{
 			if (UnrealAudioDevice->Initialize())
 			{
-				bSuccess = true;
-				InitializeTests(this);
+				// TODO: Use some sort of user savings feature to store which audio device index (or GUID) to use
+				uint32 DefaultDeviceIndex;
+				if (!UnrealAudioDevice->GetDefaultOutputDeviceIndex(DefaultDeviceIndex))
+				{
+					UE_LOG(LogUnrealAudio, Error, TEXT("Failed to get default audio device index."));
+					return false;
+				}
+
+				FDeviceInfo DeviceInfo;
+				if (!UnrealAudioDevice->GetOutputDeviceInfo(DefaultDeviceIndex, DeviceInfo))
+				{
+					UE_LOG(LogUnrealAudio, Warning, TEXT("Failed to get audio device info."));
+					return false;
+				}
+
+				// TODO: Get this information from a settings or ini file
+				FCreateStreamParams CreateStreamParams;
+				CreateStreamParams.OutputDeviceIndex = DefaultDeviceIndex;
+				CreateStreamParams.CallbackFunction = &AudioDeviceCallbackFunc;
+				CreateStreamParams.UserData = this;
+				CreateStreamParams.CallbackBlockSize = 1024;
+
+				if (!UnrealAudioDevice->CreateStream(CreateStreamParams))
+				{
+					UE_LOG(LogUnrealAudio, Display, TEXT("Failed to create an audio device stream."));
+					return false;
+				}
+
+				InitializeDeviceTests(this);
+
+				if (!UnrealAudioDevice->StartStream())
+				{
+					UE_LOG(LogUnrealAudio, Display, TEXT("Failed to start the audio device stream."));
+					return false;
+				}
+
+				UE_LOG(LogUnrealAudio, Log, TEXT("Succeeded creating output audio device..."));
 			}
 			else
 			{
 				UE_LOG(LogUnrealAudio, Error, TEXT("Failed to initialize audio device module."));
 				UnrealAudioDevice = nullptr;
+				return false;
 			}
 		}
 		else
 		{
 			UnrealAudioDevice = CreateDummyDeviceModule();
+			return true;
 		}
 
-		bSuccess = LoadSoundFileLib();
-
-		return bSuccess;
+		return true;
 	}
 
 	void FUnrealAudioModule::Shutdown()
 	{
+		MainThreadChecker.CheckThread();
+
 		// Block shutdown until the number of background tasks goes to zero before shutting down the audio module
 		static const int32 BackgroundTaskTimeout = 500;
 		int32 BackgroundTaskWaitCount = 0;
@@ -85,10 +222,12 @@ namespace UAudio
 			UE_LOG(LogUnrealAudio, Error, TEXT("Timed out while waiting for background tasks to finish when shutting down unreal audio module."));
 		}
 
-		if (UnrealAudioDevice)
-		{
-			UnrealAudioDevice->Shutdown();
-		}
+		SoundFileManager.Shutdown();
+
+		ShutdownAudioDevice();
+
+		// Stop the system thread
+		Stop();
 
 		EDeviceApi::Type DeviceApi;
 		if (UnrealAudioDevice && UnrealAudioDevice->GetDevicePlatformApi(DeviceApi) && DeviceApi == EDeviceApi::DUMMY)
@@ -98,6 +237,21 @@ namespace UAudio
 		}
 
 		ShutdownSoundFileLib();
+
+	}
+
+	void FUnrealAudioModule::ShutdownAudioDevice()
+	{
+		check(UnrealAudioDevice);
+		bool bSuccess = UnrealAudioDevice->StopStream();
+		check(bSuccess);
+		bSuccess = UnrealAudioDevice->ShutdownStream();
+		check(bSuccess);
+	}
+
+	int32 FUnrealAudioModule::GetNumBackgroundTasks() const
+	{
+		return NumBackgroundTasks.GetValue();
 	}
 
 	FName FUnrealAudioModule::GetDefaultDeviceModuleName() const
@@ -122,6 +276,222 @@ namespace UAudio
 		NumBackgroundTasks.Decrement();
 	}
 
+	double FUnrealAudioModule::GetCurrentTimeSec() const
+	{
+		return AudioSystemTimeSec;
+	}
+
+	bool FUnrealAudioModule::Init()
+	{
+		AudioThreadChecker.InitThread();
+		return true;
+	}
+
+	uint32 FUnrealAudioModule::Run()
+	{
+		AudioThreadChecker.CheckThread();
+
+		// Initialize the audio system time to 0.0. This is a time
+		// which will be used to synchronize timer-events in various audio subsystems and 
+		// represents the up-time of the audio system thread.
+		AudioSystemTimeSec = 0.0;
+		double PreviousTime = FPlatformTime::Seconds();
+
+		while (!bIsStoppingSystemThread)
+		{
+			double StartTime = FPlatformTime::Seconds();
+
+			// Update the audio system time before updating any sub-systems
+			AudioSystemTimeSec += (StartTime - PreviousTime);
+			PreviousTime = StartTime;
+
+			// Update any state from previous commands
+			VoiceManager.Update();
+
+			ExecuteAudioThreadCommands();
+
+			// call the update system test function for any test code that needs to be performed in system thread
+			FUnrealAudioModule::UpdateSystemTests();
+
+			// To sleep the audio system thread for a set amount of time at regular intervals
+			// we need to subtract how long the update took from a set delta update time
+			double DeltaTime = FPlatformTime::Seconds() - StartTime;
+			double TimeLeft = SystemThreadUpdateTime - DeltaTime;
+
+			// If we still have time left, then sleep for whatever amount that is
+			// if we don't have time left, then the system update took way too long and we 
+			// want to log the performance issue.
+			if (TimeLeft > 0.0)
+			{
+				FPlatformProcess::Sleep((float)TimeLeft);
+			}
+			else
+			{
+				UE_LOG(LogUnrealAudio,
+					Warning,
+					TEXT("Audio system thread update took longer than %.2f seconds (%.2f)"),
+					SystemThreadUpdateTime,
+					DeltaTime
+				);
+			}
+		}
+
+		return 0;
+	}
+
+	void FUnrealAudioModule::Stop()
+	{
+		bIsStoppingSystemThread = true;
+	}
+
+	void FUnrealAudioModule::SendAudioThreadCommand(const FCommand& Command)
+	{
+		MainThreadChecker.CheckThread();
+
+		AudioThreadCommandQueue.Enqueue(Command);
+	}
+
+	void FUnrealAudioModule::SendMainThreadCommand(const FCommand& Command)
+	{
+		AudioThreadChecker.CheckThread();
+
+		MainThreadCommandQueue.Enqueue(Command);
+	}
+
+	void FUnrealAudioModule::ExecuteAudioThreadCommands()
+	{
+		AudioThreadChecker.CheckThread();
+
+		FCommand Command;
+		while (AudioThreadCommandQueue.Dequeue(Command))
+		{
+			switch (Command.Id)
+			{
+				default:
+				case EAudioThreadCommand::NONE:
+
+				// Voice commands
+				//////////////////////////////////////////////////////////////////////////
+
+				case EAudioThreadCommand::VOICE_PLAY:
+					VoiceManager.PlayVoice(Command);
+					break;
+
+				case EAudioThreadCommand::VOICE_PAUSE:
+					VoiceManager.PauseVoice(Command);
+					break;
+
+				case EAudioThreadCommand::VOICE_STOP:
+					VoiceManager.StopVoice(Command);
+					break;
+
+				case EAudioThreadCommand::VOICE_SET_VOLUME_SCALE:
+					VoiceManager.SetVolumeScale(Command);
+					break;
+
+				case EAudioThreadCommand::VOICE_SET_PITCH_SCALE:
+					VoiceManager.SetPitchScale(Command);
+					break;
+
+				// Emitter commands
+				//////////////////////////////////////////////////////////////////////////
+
+				case EAudioThreadCommand::EMITTER_CREATE:
+					EmitterManager.CreateEmitter(Command);
+					break;
+
+				case EAudioThreadCommand::EMITTER_RELEASE:
+					EmitterManager.ReleaseEmitter(Command);
+					break;
+
+				case EAudioThreadCommand::EMITTER_SET_POSITION:
+					EmitterManager.SetEmitterPosition(Command);
+					break;
+			}
+
+		}
+	}
+
+	void FUnrealAudioModule::ExecuteMainThreadCommands()
+	{
+		FCommand Command;
+		while (MainThreadCommandQueue.Dequeue(Command))
+		{
+			switch (Command.Id)
+			{
+				default:
+				case EMainThreadCommand::NONE:
+
+				// Voice commands
+				//////////////////////////////////////////////////////////////////////////
+
+				case EMainThreadCommand::VOICE_DONE:
+					VoiceManager.NotifyVoiceDone(Command);
+					break;
+
+				case EMainThreadCommand::VOICE_REAL:
+					VoiceManager.NotifyVoiceReal(Command);
+					break;
+
+				case EMainThreadCommand::VOICE_VIRTUAL:
+					VoiceManager.NotifyVoiceVirtual(Command);
+					break;
+
+				case EMainThreadCommand::VOICE_SUSPEND:
+					VoiceManager.NotifyVoiceSuspend(Command);
+					break;
+
+			}
+
+		}
+	}
+
+	TSharedPtr<IEmitter> FUnrealAudioModule::EmitterCreate()
+	{
+		return TSharedPtr<IEmitter>(new FEmitter(this));
+	}
+
+	FEmitterManager& FUnrealAudioModule::GetEmitterManager()
+	{
+		return EmitterManager;
+	}
+
+	TSharedPtr<IVoice> FUnrealAudioModule::VoiceCreate(const FVoiceInitializationParams& Params)
+	{
+		if (!Params.SoundFile.IsValid())
+		{
+			UE_LOG(LogUnrealAudio, Error, TEXT("Must give a valid ISoundFile object when creating a voice."));
+			return TSharedPtr<IVoice>();
+		}
+		else
+		{
+			FVoiceInitializationParams InitParams;
+			InitParams.BaselinePitchScale = FMath::Clamp(Params.BaselinePitchScale, 0.0f, 4.0f);
+			InitParams.BaselineVolumeScale = FMath::Clamp(Params.BaselineVolumeScale, 0.0f, 1.0f);
+			InitParams.PriorityWeight = FMath::Max(Params.PriorityWeight, 0.0f);
+			InitParams.Emitter = Params.Emitter;
+			InitParams.SoundFile = Params.SoundFile;
+			InitParams.bIsLooping = Params.bIsLooping;
+			TSharedPtr<IVoice> NewVoice = TSharedPtr<IVoice>(new FVoice(this, InitParams));
+			return NewVoice;
+		}
+	}
+
+	FVoiceManager& FUnrealAudioModule::GetVoiceManager()
+	{
+		return VoiceManager;
+	}
+
+	bool FUnrealAudioModule::AudioDeviceCallback(struct FCallbackInfo& CallbackInfo)
+	{
+		// do any test device callbacks
+		DeviceTestCallback(CallbackInfo);
+		
+		VoiceManager.Mix(CallbackInfo);
+
+		return true;
+	}
+
 }
 
 #endif // #if ENABLE_UNREAL_AUDIO
diff --git a/Engine/Source/Runtime/UnrealAudio/Private/UnrealAudioDecode.cpp b/Engine/Source/Runtime/UnrealAudio/Private/UnrealAudioDecode.cpp
new file mode 100644
--- /dev/null
+++ b/Engine/Source/Runtime/UnrealAudio/Private/UnrealAudioDecode.cpp
@@ -0,0 +1,270 @@
+#include "UnrealAudioPrivate.h"
+#include "UnrealAudioSoundFile.h"
+#include "UnrealAudioSoundFileInternal.h"
+#include "UnrealAudioModule.h"
+#include "UnrealAudioDecode.h"
+
+#if ENABLE_UNREAL_AUDIO
+
+namespace UAudio
+{
+	static int32 NumDecodeThreads = 0;
+
+	FSoundFileDecoder::FSoundFileDecoder(FUnrealAudioModule* InAudioModule)
+		: AudioModule(InAudioModule)
+		, ThreadDecodeEvent(nullptr)
+		, bIsDecoding(false)
+		, DecodeThread(nullptr)
+	{
+		Settings = { 0 };
+	}
+
+	FSoundFileDecoder::~FSoundFileDecoder()
+	{
+		// Return the sync event from the pool
+		if (ThreadDecodeEvent)
+		{
+			FPlatformProcess::ReturnSynchEventToPool(ThreadDecodeEvent);
+		}
+	}
+
+	bool FSoundFileDecoder::Init(const FSoundFileDecoderSettings& Settings, int32 NumVoices)
+	{
+		// Create a thread decode event to pause the decoding thread when all buffers are filled
+		ThreadDecodeEvent = FPlatformProcess::GetSynchEventFromPool(false);
+
+		// We really should be using at least a triple buffer decoder, but we may want more...
+		check(Settings.NumDecodeBuffers >= 3);
+
+		for (int32 VoiceIndex = 0; VoiceIndex < NumVoices; ++VoiceIndex)
+		{
+			DecodeData.Add(FSoundFileDecodeData(Settings.NumDecodeBuffers));
+		}
+
+		FString DecodeThreadName = FString::Printf(TEXT("Audio Decode Thread %d"), NumDecodeThreads++);
+		DecodeThread = FRunnableThread::Create(this, *DecodeThreadName, 0, TPri_Normal);
+
+		return true;
+	}
+
+	bool FSoundFileDecoder::Init()
+	{
+		bIsDecoding = true;
+		return true;
+	}
+
+	uint32 FSoundFileDecoder::Run()
+	{
+		while (bIsDecoding)
+		{
+			for (int32 DataIndex = 0; DataIndex < Settings.NumDecodeBuffers; ++DataIndex)
+			{
+				FSoundFileDecodeData& Data = DecodeData[DataIndex];
+
+				// If the data is not active, continue
+				if (!Data.bIsActive)
+				{
+					continue;
+				}
+
+				int32 CurrentWriteBufferIndex = Data.CurrentWriteBufferIndex.GetValue();
+				while (CurrentWriteBufferIndex != Data.CurrentReadBufferIndex.GetValue())
+				{
+					TArray<float>& DecodeBuffer = Data.DecodedBuffers[CurrentWriteBufferIndex];
+
+					SoundFileCount RequestedNumSamples = (SoundFileCount)DecodeBuffer.Num();
+					SoundFileCount NumSamplesActuallyRead = 0;
+
+					// If the data is looping, then we need to keep reading samples from the file until we've reached the number of requested
+					// samples. For *very* short duration sounds (or for long-duration decode buffers), we may have an edge case where
+					// the sound file needs to seek from the beginning multiple times to fulfill the decode buffer size request
+					if (Data.bIsLooping)
+					{
+						// Keep looping until we've read everything we requested
+						while (true)
+						{
+							ESoundFileError::Type Result = Data.SoundFileReader->ReadSamples(DecodeBuffer.GetData(), RequestedNumSamples, NumSamplesActuallyRead);
+							check(Result == ESoundFileError::NONE);
+
+							// If we've read all the samples we wanted, then we're good
+							if (RequestedNumSamples == NumSamplesActuallyRead)
+							{
+								break;
+							}
+							// ... otherwise, we need to seek back to the beginning of the file and try to read samples that are left
+							else
+							{
+								// Seek to beginning of file
+								SoundFileCount OutFrameOffset = 0;
+								Data.SoundFileReader->SeekFrames(0, ESoundFileSeekMode::FROM_START, OutFrameOffset);
+
+								// Calculate samples left we need to read
+								RequestedNumSamples -= NumSamplesActuallyRead;
+							}
+						}
+					}
+					else
+					{
+						ESoundFileError::Type Result = Data.SoundFileReader->ReadSamples(DecodeBuffer.GetData(), RequestedNumSamples, NumSamplesActuallyRead);
+						check(Result == ESoundFileError::NONE);
+
+						// if the number of samples we read was not the number we requested, then we've reached the end of the file
+						if (RequestedNumSamples != NumSamplesActuallyRead)
+						{
+							// tell the data that this entry is no longer active
+							Data.bIsActive = false;
+
+							// and what our last sample index is
+							Data.LastSampleIndex = NumSamplesActuallyRead;
+						}
+					}
+
+					// Increment the buffer index
+					CurrentWriteBufferIndex = CurrentWriteBufferIndex % Settings.NumDecodeBuffers;
+
+					// Publish results by writing setting buffer index so other thread can read it
+					Data.CurrentWriteBufferIndex.Set(CurrentWriteBufferIndex);
+				}
+
+			}
+
+			// Once we're done decoding everything, then put this thread to sleep
+			// Any new requests (or any shutdowns) will trigger this thread to wake up.
+			ThreadDecodeEvent->Wait((uint32)-1);
+		}
+		return true;
+	}
+
+	void FSoundFileDecoder::Stop()
+	{
+		bIsDecoding = false;
+		Signal();
+	}
+	
+	void FSoundFileDecoder::Exit()
+	{
+	}
+
+	void FSoundFileDecoder::Signal()
+	{
+		check(ThreadDecodeEvent != nullptr);
+		ThreadDecodeEvent->Trigger();
+	}
+
+	void FSoundFileDecoder::InitalizeEntry(uint32 VoiceIndex, TSharedPtr<ISoundFile> InSoundFileData)
+	{
+// 		FSoundFileDecodeData& VoiceDecodeData = DecodeData[VoiceIndex];
+// 
+// 		check(VoiceDecodeData.bIsActive == false);
+// 
+// 		// Get the number of channels in our sound file we're going to decode
+// 
+// 		FSoundFileManager& SoundFileManager = AudioModule->GetSoundFileManager();
+// 
+// 		const FSoundFileDescription& Description = InSoundFileData->GetDescription();
+// 		int32 NumSamples = Settings.DecodeBufferFrames * Description.NumChannels;
+// 
+// 		// Make a new sound file reader with the sound file data
+// 		TSharedPtr<ISoundFileReader> SoundFileReader = AudioModule->CreateSoundFileReader();
+// 
+// 		SoundFileReader->SetSoundFileData(InSoundFileData);
+// 
+// 		uint32 SoundFileId;
+// 		InSoundFileData->GetId(SoundFileId);
+// 
+// 		VoiceDecodeData.SoundFileHandle = FSoundFileHandle(FEntityHandle(SoundFileId));
+// 
+// 		VoiceDecodeData.SoundFileReader = SoundFileReader;
+// 		VoiceDecodeData.CurrentReadBufferIndex.Set(0);
+// 		VoiceDecodeData.CurrentWriteBufferIndex.Set(1);
+// 		VoiceDecodeData.CurrentReadSampleIndex = 0;
+// 
+// 		// Zero out the decoded data from any previous decoding sessions
+// 		for (int32 BufferIndex = 0; BufferIndex < Settings.NumDecodeBuffers; ++BufferIndex)
+// 		{
+// 			VoiceDecodeData.DecodedBuffers[BufferIndex].Reset(NumSamples);
+// 		}
+// 
+// 		// Publish that this entry is now active
+// 		VoiceDecodeData.bIsActive = true;
+// 
+// 		// Wake up the writing thread in case its asleep
+// 		Signal();
+	}
+
+	void FSoundFileDecoder::ClearEntry(uint32 Index)
+	{
+// 		FSoundFileDecodeData& VoiceDecodeData = DecodeData[Index];
+// 
+// 		check(VoiceDecodeData.bIsActive == true);
+// 
+// 		// inactivate the entry
+// 		VoiceDecodeData.bIsActive = false;
+// 		VoiceDecodeData.SoundFileReader = nullptr;
+	}
+
+	bool FSoundFileDecoder::GetDecodedAudioData(uint32 VoiceIndex, TArray<float>& AudioData)
+	{
+// 		FSoundFileDecodeData& VoiceDecodeData = DecodeData[VoiceIndex];
+// 
+// 		check(VoiceDecodeData.bIsActive == true);
+// 
+// 		int32 CurrentReadBufferIndex = VoiceDecodeData.CurrentReadBufferIndex.GetValue();
+// 
+// 		TArray<float>& CurrentReadBuffer = VoiceDecodeData.DecodedBuffers[CurrentReadBufferIndex];
+// 		int32 ReadBufferSize = CurrentReadBuffer.Num();
+// 
+// 		// Get the current sample index
+// 		int32& CurrentReadSampleIndex = VoiceDecodeData.CurrentReadSampleIndex;
+// 		int32 NumRequestedSamples = AudioData.Num();
+// 
+// 		for (int32 OutputSample = 0; OutputSample < NumRequestedSamples; ++OutputSample)
+// 		{
+// 			// Check the non-terminating conditions
+// 			if (VoiceDecodeData.bIsActive || (CurrentReadBufferIndex != VoiceDecodeData.CurrentWriteBufferIndex.GetValue()) || CurrentReadSampleIndex != VoiceDecodeData.LastSampleIndex)
+// 			{
+// 				// Copy the decoded data to the output buffer
+// 				AudioData[OutputSample] = CurrentReadBuffer[CurrentReadSampleIndex++];
+// 
+// 				// if the current read sample index is equal to the current read buffer size,
+// 				// then we'll need to bump the read buffer index and signal the decoding thread that 
+// 				// it needs to wake up and start decoding some audio (if it's asleep)
+// 				if (CurrentReadSampleIndex == ReadBufferSize)
+// 				{
+// 					CurrentReadBufferIndex = (CurrentReadBufferIndex + 1) % Settings.NumDecodeBuffers;
+// 
+// 					// first check to see if our new read buffer index is actually the same as the current
+// 					// write buffer index, if it is, then we'll need to (unfortunately) wait. This should
+// 					// rarely happen.
+// 
+// 					while (CurrentReadBufferIndex == VoiceDecodeData.CurrentWriteBufferIndex.GetValue())
+// 					{
+// 						// spin
+// 					}
+// 
+// 					// free up the last read buffer for writing
+// 					VoiceDecodeData.CurrentReadBufferIndex.Set(CurrentReadBufferIndex);
+// 
+// 					// Reset the read sample index 
+// 					CurrentReadSampleIndex = 0;
+// 
+// 					// And update the read buffer
+// 					CurrentReadBuffer = VoiceDecodeData.DecodedBuffers[CurrentReadBufferIndex];
+// 				}
+// 			}
+// 			else
+// 			{
+// 				// we're done decoding this sound file... it's a "one-shot"
+// 				return true;
+// 			}
+// 		}
+// 
+// 		// we still got more audio to decode on this sound file
+		return false;
+
+	}
+
+
+}
+
+#endif // #if ENABLE_UNREAL_AUDIO
\ No newline at end of file
diff --git a/Engine/Source/Runtime/UnrealAudio/Private/UnrealAudioDecode.h b/Engine/Source/Runtime/UnrealAudio/Private/UnrealAudioDecode.h
new file mode 100644
--- /dev/null
+++ b/Engine/Source/Runtime/UnrealAudio/Private/UnrealAudioDecode.h
@@ -0,0 +1,83 @@
+// Copyright 1998-2015 Epic Games, Inc. All Rights Reserved.
+
+#pragma once 
+
+#include "UnrealAudioModule.h"
+#include "UnrealAudioHandles.h"
+#include "UnrealAudioSoundFileManager.h"
+#include "UnrealAudioTypes.h"
+#include "UnrealAudioSoundFile.h"
+#include "UnrealAudioSoundFileInternal.h"
+
+#if ENABLE_UNREAL_AUDIO
+
+namespace UAudio
+{
+	struct FSoundFileDecoderSettings
+	{
+		int32 NumDecodeBuffers;
+		int32 DecodeBufferFrames;
+	};
+
+	class FSoundFileDecoder : public FRunnable
+	{
+	public:
+		FSoundFileDecoder(FUnrealAudioModule* InAudioModule);
+		~FSoundFileDecoder();
+
+		// FRunnable
+		bool Init() override;
+		uint32 Run() override;
+		void Stop() override;
+		void Exit() override;
+
+		bool Init(const FSoundFileDecoderSettings& Settings, int32 NumVoices);
+		void InitalizeEntry(uint32 Index, TSharedPtr<ISoundFile> InSoundFileData);
+		void ClearEntry(uint32 Index);
+
+		bool GetDecodedAudioData(uint32 Index, TArray<float>& AudioData);
+
+	private:
+
+		void Signal();
+
+		struct FSoundFileDecodeData
+		{
+			FSoundFileHandle SoundFileHandle;
+			TSharedPtr<FSoundFileReader> SoundFileReader;
+			TArray<TArray<float>> DecodedBuffers;
+			FThreadSafeBool bIsActive;
+			int32 CurrentReadSampleIndex;
+			int32 LastSampleIndex;
+			int32 NumChannels;
+			bool bIsLooping;
+			FThreadSafeCounter CurrentWriteBufferIndex;
+			FThreadSafeCounter CurrentReadBufferIndex;
+
+			FSoundFileDecodeData(int32 NumDecodeBuffers)
+				: bIsActive(false)
+				, CurrentReadSampleIndex(0)
+				, LastSampleIndex(0)
+				, NumChannels(1)
+				, bIsLooping(false)
+				, CurrentWriteBufferIndex(0)
+				, CurrentReadBufferIndex(0)
+			{
+				// Initialize this entries decode buffers
+				DecodedBuffers.Init(TArray<float>(), NumDecodeBuffers);
+			}
+		};
+
+		FUnrealAudioModule* AudioModule;
+		FEvent* ThreadDecodeEvent;
+		FSoundFileDecoderSettings Settings;
+		TArray<FSoundFileDecodeData> DecodeData;
+		FThreadSafeBool bIsDecoding;
+		FRunnableThread* DecodeThread;
+	};
+}
+
+
+
+#endif // #if ENABLE_UNREAL_AUDIO
+
diff --git a/Engine/Source/Runtime/UnrealAudio/Private/UnrealAudioEmitter.cpp b/Engine/Source/Runtime/UnrealAudio/Private/UnrealAudioEmitter.cpp
new file mode 100644
--- /dev/null
+++ b/Engine/Source/Runtime/UnrealAudio/Private/UnrealAudioEmitter.cpp
@@ -0,0 +1,76 @@
+// Copyright 1998-2015 Epic Games, Inc. All Rights Reserved.
+
+#include "UnrealAudioPrivate.h"
+#include "UnrealAudioModule.h"
+#include "UnrealAudioEmitterInternal.h"
+
+#define UNREAL_AUDIO_EMITTER_COMMAND_QUEUE_SIZE (50)
+
+#if ENABLE_UNREAL_AUDIO
+
+namespace UAudio
+{
+	/************************************************************************/
+	/* FEmitter implementation                                              */
+	/************************************************************************/
+
+	FEmitter::FEmitter(FUnrealAudioModule* InParentModule)
+		: Position(0)
+		, AudioModule(InParentModule)
+	{
+		check(AudioModule != nullptr);
+		FEmitterManager& EmitterManager = AudioModule->GetEmitterManager();
+		EmitterHandle = EmitterManager.CreateEmitter();
+	}
+
+	FEmitter::~FEmitter()
+	{
+		// Release this emitter if it hasn't already been released
+		Release();
+	}
+
+	ESystemError::Type FEmitter::SetPosition(const FVector& InPosition)
+	{
+		if (InPosition != Position)
+		{
+			Position = InPosition;
+			FEmitterManager& EmitterManager = AudioModule->GetEmitterManager();
+			return EmitterManager.SetEmitterPosition(EmitterHandle, InPosition);
+		}
+
+		return ESystemError::NONE;
+	}
+
+	ESystemError::Type FEmitter::GetPosition(FVector& OutPosition)
+	{
+		OutPosition = Position;
+		return ESystemError::NONE;
+	}
+
+	ESystemError::Type FEmitter::Release()
+	{
+		FEmitterManager& EmitterManager = AudioModule->GetEmitterManager();
+		if (EmitterHandle.IsInitialized())
+		{
+			ESystemError::Type Result = EmitterManager.ReleaseEmitter(EmitterHandle);
+			EmitterHandle.Id = INDEX_NONE;
+			return Result;
+		}
+		return ESystemError::INVALID_HANDLE;
+	}
+
+	uint32 FEmitter::GetId() const
+	{
+		return EmitterHandle.Id;
+	}
+
+	FEmitterHandle FEmitter::GetHandle() const
+	{
+		return EmitterHandle;
+	}
+
+
+}
+
+#endif // #if ENABLE_UNREAL_AUDIO
+
diff --git a/Engine/Source/Runtime/UnrealAudio/Private/UnrealAudioEmitterInternal.h b/Engine/Source/Runtime/UnrealAudio/Private/UnrealAudioEmitterInternal.h
new file mode 100644
--- /dev/null
+++ b/Engine/Source/Runtime/UnrealAudio/Private/UnrealAudioEmitterInternal.h
@@ -0,0 +1,57 @@
+#pragma once 
+
+#include "UnrealAudioModule.h"
+#include "UnrealAudioPrivate.h"
+#include "UnrealAudioTypes.h"
+#include "UnrealAudioEmitter.h"
+#include "UnrealAudioPrivate.h"
+#include "UnrealAudioEntityManager.h"
+#include "UnrealAudioUtilities.h"
+
+#if ENABLE_UNREAL_AUDIO
+
+namespace UAudio
+{
+	/**
+	* FEmitter
+	* Concrete implementation of IEmitter
+	*/
+	class FEmitter : public IEmitter
+	{
+	public:
+		FEmitter(class FUnrealAudioModule* ParentModule);
+		~FEmitter() override;
+
+		// IEmitter interface
+		ESystemError::Type SetPosition(const FVector& InPosition) override;
+		ESystemError::Type GetPosition(FVector& OutPosition) override;
+		ESystemError::Type Release() override;
+		uint32 GetId() const override;
+
+		FEmitterHandle GetHandle() const;
+
+	private:
+		/** The current position of the emitter. */
+		FVector Position;
+
+		/** Handle to the emitter on the audio thread. */
+		FEmitterHandle EmitterHandle;
+
+		/** The parent unreal module instance. */
+		FUnrealAudioModule* AudioModule;
+	};
+
+
+	namespace EEmitterCommand
+	{
+		enum Type
+		{
+			NONE,			// No emitter command
+			CREATE,			// create a new emitter
+			RELEASE,		// release the emitter
+			SET_POSITION,	// set the emitter position
+		};
+	}
+}
+
+#endif // #if ENABLE_UNREAL_AUDIO
\ No newline at end of file
diff --git a/Engine/Source/Runtime/UnrealAudio/Private/UnrealAudioEmitterManager.cpp b/Engine/Source/Runtime/UnrealAudio/Private/UnrealAudioEmitterManager.cpp
new file mode 100644
--- /dev/null
+++ b/Engine/Source/Runtime/UnrealAudio/Private/UnrealAudioEmitterManager.cpp
@@ -0,0 +1,142 @@
+// Copyright 1998-2015 Epic Games, Inc. All Rights Reserved.
+
+#include "UnrealAudioPrivate.h"
+#include "UnrealAudioModule.h"
+#include "UnrealAudioEmitterInternal.h"
+#include "UnrealAudioEmitterManager.h"
+
+#define UNREAL_AUDIO_EMITTER_COMMAND_QUEUE_SIZE (50)
+
+#if ENABLE_UNREAL_AUDIO
+
+namespace UAudio
+{
+	/************************************************************************/
+	/* Main Thread Functions												*/
+	/************************************************************************/
+
+	FEmitterManager::FEmitterManager(FUnrealAudioModule* InAudioModule)
+		: EntityManager(500)
+		, AudioModule(InAudioModule)
+	{
+		check(AudioModule);
+	}
+
+	FEmitterManager::~FEmitterManager()
+	{
+	}
+
+	FEmitterHandle FEmitterManager::CreateEmitter()
+	{
+		FEmitterHandle EmitterHandle(EntityManager.CreateEntity());
+
+		AudioModule->SendAudioThreadCommand(FCommand(EAudioThreadCommand::EMITTER_CREATE, EmitterHandle));
+
+		return EmitterHandle;
+	}
+
+	ESystemError::Type FEmitterManager::ReleaseEmitter(const FEmitterHandle& EmitterHandle)
+	{
+		EntityManager.ReleaseEntity(EmitterHandle);
+
+		AudioModule->SendAudioThreadCommand(FCommand(EAudioThreadCommand::EMITTER_RELEASE, EmitterHandle.Id));
+
+		return ESystemError::NONE;
+	}
+
+	ESystemError::Type FEmitterManager::SetEmitterPosition(const FEmitterHandle& EmitterHandle, const FVector& Position)
+	{
+		if (EntityManager.IsValidEntity(EmitterHandle))
+		{
+			return ESystemError::INVALID_HANDLE;
+		}
+
+		AudioModule->SendAudioThreadCommand(FCommand(EAudioThreadCommand::EMITTER_SET_POSITION, EmitterHandle, Position.X, Position.Y, Position.Z));
+
+		return ESystemError::NONE;
+	}
+
+	bool FEmitterManager::IsValid(const FEmitterHandle& EmitterHandle) const
+	{
+		return EntityManager.IsValidEntity(EmitterHandle);
+	}
+
+	/************************************************************************/
+	/* Audio Thread Functions												*/
+	/************************************************************************/
+
+	void FEmitterManager::CreateEmitter(const FCommand& Command)
+	{
+		check(Command.NumArguments == 1);
+		check(Command.Arguments[0].DataType == ECommandData::UINT32);
+
+		FEmitterHandle EmitterHandle = Command.Arguments[0].Data.UnsignedInt32;
+		uint32 EmitterIndex = EmitterHandle.GetIndex();
+
+		if (EmitterIndex < (uint32)EmitterData.Num())
+		{
+			FEmitterData& Data = EmitterData[EmitterIndex];
+			check(!Data.EmitterHandle.IsInitialized());
+			Data.EmitterHandle = EmitterHandle;
+		}
+		else
+		{
+			check(EmitterIndex == (uint32)EmitterData.Num());
+			EmitterData.Add(FEmitterData(EmitterHandle));
+		}
+	}
+
+	void FEmitterManager::ReleaseEmitter(const FCommand& Command)
+	{
+		check(Command.NumArguments == 1);
+		check(Command.Arguments[0].DataType == ECommandData::UINT32);
+
+		FEmitterHandle EmitterHandle = Command.Arguments[0].Data.UnsignedInt32;
+		uint32 EmitterIndex = EmitterHandle.GetIndex();
+		ValidateEmitterEntry(EmitterHandle, EmitterIndex);
+
+		EmitterData[EmitterIndex].EmitterHandle.Id = INDEX_NONE;
+	}
+
+	void FEmitterManager::SetEmitterPosition(const FCommand& Command)
+	{
+		check(Command.NumArguments == 4);
+		check(Command.Arguments[0].DataType == ECommandData::UINT32);
+		check(Command.Arguments[1].DataType == ECommandData::FLOAT_32);
+		check(Command.Arguments[2].DataType == ECommandData::FLOAT_32);
+		check(Command.Arguments[3].DataType == ECommandData::FLOAT_32);
+
+		FEmitterHandle EmitterHandle = Command.Arguments[0].Data.UnsignedInt32;
+
+		FVector Position;
+		Position.X = Command.Arguments[1].Data.Float32Val;
+		Position.Y = Command.Arguments[2].Data.Float32Val;
+		Position.Z = Command.Arguments[3].Data.Float32Val;
+
+		uint32 EmitterIndex = EmitterHandle.GetIndex();
+		ValidateEmitterEntry(EmitterHandle, EmitterIndex);
+
+		// Make sure we are getting a new position -- this should be checked before sending a command
+		check(EmitterData[EmitterIndex].Position != Position);
+
+		// Update the position at the index
+		EmitterData[EmitterIndex].Position = Position;
+	}
+
+	const FVector* FEmitterManager::GetEmitterPosition(const FEmitterHandle& EmitterHandle) const
+	{
+		uint32 EmitterIndex = EmitterHandle.GetIndex();
+		ValidateEmitterEntry(EmitterHandle, EmitterIndex);
+		return &EmitterData[EmitterIndex].Position;
+	}
+
+	void FEmitterManager::ValidateEmitterEntry(const FEmitterHandle& EmitterHandle, uint32 EmitterIndex) const
+	{
+		check(EmitterIndex < (uint32)EmitterData.Num());
+		check(EmitterData[EmitterIndex].EmitterHandle.Id == EmitterHandle.Id);
+	}
+
+}
+
+#endif // #if ENABLE_UNREAL_AUDIO
+
diff --git a/Engine/Source/Runtime/UnrealAudio/Private/UnrealAudioEmitterManager.h b/Engine/Source/Runtime/UnrealAudio/Private/UnrealAudioEmitterManager.h
new file mode 100644
--- /dev/null
+++ b/Engine/Source/Runtime/UnrealAudio/Private/UnrealAudioEmitterManager.h
@@ -0,0 +1,76 @@
+#pragma once 
+
+#include "UnrealAudioPrivate.h"
+#include "UnrealAudioModule.h"
+#include "UnrealAudioTypes.h"
+#include "UnrealAudioEmitter.h"
+#include "UnrealAudioEntityManager.h"
+#include "UnrealAudioEmitterInternal.h"
+#include "UnrealAudioUtilities.h"
+
+#if ENABLE_UNREAL_AUDIO
+
+namespace UAudio
+{
+	/**
+	* FEmitterManager
+	* Manages emitter handles, messaging, and updating.
+	*/
+	class FEmitterManager
+	{
+	public:
+		FEmitterManager(FUnrealAudioModule* InAudioModule);
+		~FEmitterManager();
+
+		/************************************************************************/
+		/* Main Thread Functions												*/
+		/************************************************************************/
+
+		FEmitterHandle FEmitterManager::CreateEmitter();
+		ESystemError::Type FEmitterManager::ReleaseEmitter(const FEmitterHandle& EmitterHandle);
+		ESystemError::Type FEmitterManager::SetEmitterPosition(const FEmitterHandle& EmitterHandle, const FVector& Position);
+		bool IsValid(const FEmitterHandle& EmitterHandle) const;
+
+		/************************************************************************/
+		/* Audio System Thread Functions										*/
+		/************************************************************************/
+
+		void CreateEmitter(const FCommand& Command);
+		void ReleaseEmitter(const FCommand& Command);
+		void SetEmitterPosition(const FCommand& Command);
+		const FVector* GetEmitterPosition(const FEmitterHandle& EmitterHandle) const;
+
+	private:
+
+		void ValidateEmitterEntry(const FEmitterHandle& EmitterHandle, uint32 EmitterIndex) const;
+
+		/** Struct to hold data about emitters in audio system thread. */
+		struct FEmitterData
+		{
+			FEmitterHandle EmitterHandle;
+			FVector Position;
+
+			FEmitterData()
+				: Position(0)
+			{}
+
+			FEmitterData(FEmitterHandle InEmitterHandle)
+				: EmitterHandle(InEmitterHandle)
+				, Position(0)
+			{}
+		};
+
+		/** Array of emitter data in audio system thread. */
+		TArray<FEmitterData> EmitterData;
+
+		/** Handles emitter entity handles. */
+		FEntityManager EntityManager;
+
+		/** Parent audio module for the emitter manager. */
+		FUnrealAudioModule* AudioModule;
+
+	};
+
+}
+
+#endif // #if ENABLE_UNREAL_AUDIO
\ No newline at end of file
diff --git a/Engine/Source/Runtime/UnrealAudio/Private/UnrealAudioEntityManager.h b/Engine/Source/Runtime/UnrealAudio/Private/UnrealAudioEntityManager.h
new file mode 100644
--- /dev/null
+++ b/Engine/Source/Runtime/UnrealAudio/Private/UnrealAudioEntityManager.h
@@ -0,0 +1,202 @@
+// Copyright 1998-2015 Epic Games, Inc. All Rights Reserved.
+
+#pragma once 
+
+#include "UnrealAudioModule.h"
+
+#if ENABLE_UNREAL_AUDIO
+
+namespace UAudio
+{
+	static const uint32 ENTITY_INDEX_BITS = 24;
+	static const uint32 ENTITY_INDEX_MASK = (1 << ENTITY_INDEX_BITS) - 1;
+	static const uint32 ENTITY_GENERATION_BITS = 8;
+	static const uint32 ENTITY_GENERATION_MASK = (1 << ENTITY_GENERATION_BITS) - 1;
+	static const uint32 ENTITY_INDEX_INVALID = INDEX_NONE;
+
+	/** An entity handle. */
+	struct FEntityHandle
+	{
+
+	public:
+
+		/** Default constructor. */
+
+		FEntityHandle()
+			: Id(INDEX_NONE)
+		{}
+
+		/**
+		 * Constructor.
+		 *
+		 * @param	InId	Identifier for the in.
+		 */
+
+		FEntityHandle(uint32 InId)
+			: Id(InId)
+		{}
+
+		/**
+		 * Constructor.
+		 *
+		 * @param	Other	The other.
+		 */
+
+		FEntityHandle(const FEntityHandle& Other)
+			: Id(Other.Id)
+		{}
+
+		/**
+		 * Query if this object is initialized.
+		 *
+		 * @return	true if initialized, false if not.
+		 */
+
+		bool IsInitialized() const
+		{
+			return Id != INDEX_NONE;
+		}
+
+		/**
+		 * Gets the index.
+		 *
+		 * @return	The index.
+		 */
+
+		uint32 GetIndex() const
+		{
+			return Id & ENTITY_INDEX_MASK;
+		}
+
+		/**
+		 * Gets the generation.
+		 *
+		 * @return	The generation.
+		 */
+
+		uint32 GetGeneration() const
+		{
+			return (Id >> ENTITY_GENERATION_BITS) & ENTITY_GENERATION_MASK;
+		}
+
+		/** The identifier. */
+		uint32 Id;
+	};
+
+	/**
+	 * Manages entity handles for weak-referencing of resources between threads and processes.
+	 */
+	class FEntityManager
+	{
+	public:
+
+		/**
+		 * Constructor.
+		 *
+		 * @param	InMinNumFreeIndices	The minimum number free indices.
+		 */
+
+		FEntityManager(uint32 InMinNumFreeIndices)
+			: NumFreeEntities(0)
+			, MinNumFreeIndices(InMinNumFreeIndices)
+		{
+
+		}
+
+		/** Destructor. */
+		~FEntityManager()
+		{
+		}
+
+		/**
+		 * Creates the entity.
+		 *
+		 * @return	The new entity.
+		 */
+
+		FEntityHandle CreateEntity()
+		{
+			uint32 NewEntityIndex = ENTITY_INDEX_INVALID;
+			if (NumFreeEntities > MinNumFreeIndices)
+			{
+				FreeEntityIndices.Dequeue(NewEntityIndex);
+				--NumFreeEntities;
+			}
+			else
+			{
+				NewEntityIndex = EntityGenerations.Num();
+				EntityGenerations.Add(0);
+				check(NewEntityIndex < (1 << ENTITY_INDEX_BITS));
+			}
+			return CreateHandle(NewEntityIndex, EntityGenerations[NewEntityIndex]);
+		}
+
+		/**
+		 * Releases the entity described by Handle.
+		 *
+		 * @param	Handle	The handle.
+		 */
+
+		void ReleaseEntity(const FEntityHandle& Handle)
+		{
+			uint32 EntityIndex = Handle.GetIndex();
+			uint8 Generation = Handle.GetGeneration();
+			EntityGenerations[EntityIndex] = ++Generation;
+			++NumFreeEntities;
+			FreeEntityIndices.Enqueue(EntityIndex);
+		}
+
+		/**
+		 * Query if 'Handle' is valid entity.
+		 *
+		 * @param	Handle	The handle.
+		 *
+		 * @return	true if valid entity, false if not.
+		 */
+
+		bool IsValidEntity(FEntityHandle Handle) const
+		{
+			if (Handle.Id == ENTITY_INDEX_INVALID)
+			{
+				return false;
+			}
+
+			uint32 EntityIndex = Handle.GetIndex();
+			uint8 Generation = Handle.GetGeneration();
+			return EntityGenerations[EntityIndex] == Generation;
+		}
+
+	private:
+
+		/**
+		 * Creates a handle.
+		 *
+		 * @param	EntityIndex			Zero-based index of the entity.
+		 * @param	EntityGeneration	The entity generation.
+		 *
+		 * @return	The new handle.
+		 */
+
+		FEntityHandle CreateHandle(uint32 EntityIndex, uint8 EntityGeneration) const
+		{
+			return (EntityIndex | (EntityGeneration << ENTITY_INDEX_BITS));
+		}
+
+		/** Number of free entities. */
+		uint32 NumFreeEntities;
+
+		/** The minimum number free indices. */
+		uint32 MinNumFreeIndices;
+
+		/** The free entity indices. */
+		TQueue<uint32> FreeEntityIndices;
+
+		/** The entity generations. */
+		TArray<uint8> EntityGenerations;
+	};
+
+}
+
+#endif // #if ENABLE_UNREAL_AUDIO
+
+
diff --git a/Engine/Source/Runtime/UnrealAudio/Private/UnrealAudioHandles.h b/Engine/Source/Runtime/UnrealAudio/Private/UnrealAudioHandles.h
new file mode 100644
--- /dev/null
+++ b/Engine/Source/Runtime/UnrealAudio/Private/UnrealAudioHandles.h
@@ -0,0 +1,57 @@
+// Copyright 1998-2015 Epic Games, Inc. All Rights Reserved.
+
+#pragma once 
+
+#include "UnrealAudioModule.h"
+#include "UnrealAudioEntityManager.h"
+
+#if ENABLE_UNREAL_AUDIO
+
+namespace UAudio
+{
+	struct FSoundFileHandle : public FEntityHandle
+	{
+		FSoundFileHandle()
+			: FEntityHandle()
+		{}
+
+		FSoundFileHandle(const FSoundFileHandle& InEntityHandle)
+			: FEntityHandle(InEntityHandle)
+		{}
+
+		FSoundFileHandle(const FEntityHandle& InEntityHandle)
+			: FEntityHandle(InEntityHandle)
+		{}
+	};
+
+	struct FVoiceHandle : public FEntityHandle
+	{
+		FVoiceHandle()
+			: FEntityHandle()
+		{}
+
+		FVoiceHandle(const FVoiceHandle& InEntityHandle)
+			: FEntityHandle(InEntityHandle)
+		{}
+
+		FVoiceHandle(const FEntityHandle& InEntityHandle)
+			: FEntityHandle(InEntityHandle)
+		{}
+	};
+
+	struct FEmitterHandle : public FEntityHandle
+	{
+		FEmitterHandle()
+			: FEntityHandle()
+		{}
+
+		FEmitterHandle(const FEntityHandle& InEntityHandle)
+			: FEntityHandle(InEntityHandle)
+		{}
+	};
+
+}
+
+#endif // #if ENABLE_UNREAL_AUDIO
+
+
diff --git a/Engine/Source/Runtime/UnrealAudio/Private/UnrealAudioPitchManager.cpp b/Engine/Source/Runtime/UnrealAudio/Private/UnrealAudioPitchManager.cpp
new file mode 100644
--- /dev/null
+++ b/Engine/Source/Runtime/UnrealAudio/Private/UnrealAudioPitchManager.cpp
@@ -0,0 +1,134 @@
+// Copyright 1998-2015 Epic Games, Inc. All Rights Reserved.
+
+#include "UnrealAudioPrivate.h"
+#include "UnrealAudioModule.h"
+#include "UnrealAudioPitchManager.h"
+
+#if ENABLE_UNREAL_AUDIO
+
+namespace UAudio
+{
+	FPitchManager::FPitchManager(FUnrealAudioModule* InAudioModule)
+		: AudioModule(InAudioModule)
+		, EntryCount(0)
+		, LastTimeSeconds(0.0)
+	{
+		check(AudioModule);
+	}
+
+	void FPitchManager::Init(uint32 NumElements)
+	{
+		Baselines.Init(1.0f, NumElements);
+		PitchScaleData.Init(NumElements);
+		PitchProducts.Init(1.0f, NumElements);
+		DurationSeconds.Init(0.0f, NumElements);
+		CurrentTimeSeconds.Init(0.0f, NumElements);
+	}
+
+	void FPitchManager::Update()
+	{
+		double CurrentTimeSec = AudioModule->GetCurrentTimeSec();
+		LastTimeSeconds = CurrentTimeSec;
+
+		// If there are no pitches to update, then no need to update anything
+		if (EntryCount == 0)
+		{
+			return;
+		}
+
+		float DeltaTime = (float)(CurrentTimeSec - LastTimeSeconds);
+		float Fraction = 0.0f;
+		float Result = 0.0f;
+
+		int32 TotalCount = Baselines.Num();
+		int32 UpdateCount = 0;
+		
+		// Update all dynamic volume parameters in one parallelized loop.
+		// TODO: implement SSE
+
+		for (int32 Index = 0; Index < TotalCount && UpdateCount < EntryCount; ++Index)
+		{
+			// If the value of baseline is less than 0.0f, it's been a released volume entry
+			if (Baselines[Index] < 0.0f)
+			{
+				continue;
+			}
+
+			++UpdateCount;
+
+			float Product = Baselines[Index] * PitchScaleData.Compute(Index, CurrentTimeSec);
+			PitchProducts[Index] = Product;
+
+			// increment the current time by the delta time
+			// Note: this is not 100% accurate for highly dynamic pitches but is a close approximation and "good enough"
+			CurrentTimeSeconds[Index] += (DeltaTime * Product);
+		}
+
+		// Make sure we updated everything!
+		DEBUG_AUDIO_CHECK(UpdateCount == EntryCount);
+	}
+
+	void FPitchManager::InitializeEntry(uint32 Index, const FPitchInitParam& PitchParams)
+	{
+		++EntryCount;
+
+		DEBUG_AUDIO_CHECK(EntryCount <= Baselines.Num());
+
+		PitchScaleData.InitEntry(Index);
+		PitchScaleData.SetValue(Index, PitchParams.PitchScale, LastTimeSeconds, PitchParams.PitchScaleTime);
+		Baselines[Index] = PitchParams.BaselinePitch;
+
+		float PitchProduct = PitchParams.BaselinePitch;
+		if (PitchParams.PitchScaleTime == 0.0f)
+		{
+			PitchProduct *= PitchParams.PitchScale;
+		}
+		PitchProducts[Index] = PitchProduct;
+		DurationSeconds[Index] = PitchParams.DurationSeconds;
+		CurrentTimeSeconds[Index] = 0.0f;
+	}
+
+	void FPitchManager::ReleaseEntry(uint32 Index)
+	{
+		DEBUG_AUDIO_CHECK(Index < (uint32)Baselines.Num());
+
+		--EntryCount;
+		DEBUG_AUDIO_CHECK(EntryCount >= 0);
+
+		// Setting baseline to -1.0f is releasing the volume data at that index
+		Baselines[Index] = -1.0f;
+	}
+
+	void FPitchManager::SetDynamicPitchScale(uint32 Index, float InPitch, float InDeltaTimeSeconds)
+	{
+		DEBUG_AUDIO_CHECK(Index < (uint32)Baselines.Num());
+		DEBUG_AUDIO_CHECK(Baselines[Index] > 0.0f);
+		PitchScaleData.SetValue(Index, InPitch, LastTimeSeconds, InDeltaTimeSeconds);
+	}
+
+	float FPitchManager::GetPitchScale(uint32 Index) const
+	{
+		DEBUG_AUDIO_CHECK(Index < (uint32)Baselines.Num());
+		DEBUG_AUDIO_CHECK(Baselines[Index] > 0.0f);
+		return PitchScaleData.CurrentValue[Index];
+	}
+
+	float FPitchManager::GetPitchProduct(uint32 Index) const
+	{
+		DEBUG_AUDIO_CHECK(Index < (uint32)Baselines.Num());
+		DEBUG_AUDIO_CHECK(Baselines[Index] > 0.0f);
+		return PitchProducts[Index];
+	}
+
+	float FPitchManager::GetPlayPercentage(uint32 Index) const
+	{
+		DEBUG_AUDIO_CHECK(Index < (uint32)Baselines.Num());
+		DEBUG_AUDIO_CHECK(Baselines[Index] > 0.0f);
+		DEBUG_AUDIO_CHECK(DurationSeconds[Index] > 0.0f);
+		return CurrentTimeSeconds[Index] / DurationSeconds[Index];
+	}
+
+}
+
+#endif // #if ENABLE_UNREAL_AUDIO
+
diff --git a/Engine/Source/Runtime/UnrealAudio/Private/UnrealAudioPitchManager.h b/Engine/Source/Runtime/UnrealAudio/Private/UnrealAudioPitchManager.h
new file mode 100644
--- /dev/null
+++ b/Engine/Source/Runtime/UnrealAudio/Private/UnrealAudioPitchManager.h
@@ -0,0 +1,104 @@
+#pragma once 
+
+#include "UnrealAudioModule.h"
+#include "UnrealAudioTypes.h"
+#include "UnrealAudioVoice.h"
+#include "UnrealAudioPrivate.h"
+#include "UnrealAudioVoiceInternal.h"
+#include "UnrealAudioUtilities.h"
+
+#if ENABLE_UNREAL_AUDIO
+
+namespace UAudio
+{
+	struct FPitchInitParam
+	{
+		float BaselinePitch;
+		float PitchScale;
+		float PitchScaleTime;
+		float DurationSeconds;
+	};
+
+	/** Manager for pitch processing. */
+	class FPitchManager
+	{
+	public:
+
+		/**
+		 * Constructor.
+		 *
+		 * @param [int]	InAudioModule	The parent audio module.
+		 */
+
+		FPitchManager(FUnrealAudioModule* InAudioModule);
+
+		/**
+		 * Initializes this object.
+		 *
+		 * @param	ReserveCount	Number of objects entries to reserve.
+		 */
+
+		void Init(uint32 ReserveCount);
+
+		/** Updates this object. */
+		void Update();
+
+		/**
+		 * Initializes the entry.
+		 *
+		 * @param	Index   	Zero-based index of the.
+		 * @param	Baseline	The baseline.
+		 */
+
+		void InitializeEntry(uint32 Index, const FPitchInitParam& PitchParams);
+
+		/**
+		 * Releases the entry described by Index.
+		 *
+		 * @param	Index	Zero-based index of the.
+		 *
+		 */
+
+		void ReleaseEntry(uint32 Index);
+
+		/**
+		 * Sets dynamic pitch.
+		 *
+		 * @param	EntityIndex			Zero-based index of the entity.
+		 * @param	PitchScale			The pitch scale.
+		 * @param	DeltaTimeSeconds	The delta time in seconds.
+		 */
+
+		void SetDynamicPitchScale(uint32 Index, float PitchScale, float DeltaTimeSeconds);
+
+		/**
+		 * Gets pitch product.
+		 *
+		 * @param	EntityIndex	Zero-based index of the entity.
+		 *
+		 * @return	The pitch product.
+		 */
+
+		float GetPitchScale(uint32 Index) const;
+
+		float GetPitchProduct(uint32 Index) const;
+
+		float GetPlayPercentage(uint32 Index) const;
+
+	private:
+
+		/** The audio module. */
+		FUnrealAudioModule* AudioModule;
+		int32 EntryCount;
+		double LastTimeSeconds;
+		TArray<float> Baselines;
+		FDynamicParamData PitchScaleData;
+		TArray<float> PitchProducts;
+		TArray<float> CurrentTimeSeconds;
+		TArray<float> DurationSeconds;
+	};
+
+
+}
+
+#endif // #if ENABLE_UNREAL_AUDIO
diff --git a/Engine/Source/Runtime/UnrealAudio/Private/UnrealAudioPrivate.h b/Engine/Source/Runtime/UnrealAudio/Private/UnrealAudioPrivate.h
--- a/Engine/Source/Runtime/UnrealAudio/Private/UnrealAudioPrivate.h
+++ b/Engine/Source/Runtime/UnrealAudio/Private/UnrealAudioPrivate.h
@@ -3,63 +3,218 @@
 #pragma once 
 
 #include "UnrealAudioModule.h"
+#include "UnrealAudioHandles.h"
 #include "UnrealAudioDeviceModule.h"
+#include "UnrealAudioEmitterInternal.h"
+#include "UnrealAudioVoiceInternal.h"
+#include "UnrealAudioEmitterManager.h"
+#include "UnrealAudioVoiceManager.h"
+#include "UnrealAudioSoundFileManager.h"
 #include "UnrealAudioTests.h"
+#include "UnrealAudioUtilities.h"
 
 #if ENABLE_UNREAL_AUDIO
 
 DECLARE_LOG_CATEGORY_EXTERN(LogUnrealAudio, Log, All);
 
+#define UA_SYSTEM_ERROR(ENUM, INFO)	(UAudio::OnSystemError(ENUM, INFO, FString(__FILE__), __LINE__))
+
+
 namespace UAudio
 {
+
+	/**
+	* EAudioThreadCommand
+	* Commands sent from main thread to audio system thread.
+	*/
+	namespace EAudioThreadCommand
+	{
+		enum Type
+		{
+			NONE = 0,
+
+			// Voice Commands
+			VOICE_PLAY,
+			VOICE_PAUSE,
+			VOICE_STOP,
+			VOICE_SET_VOLUME_SCALE,
+			VOICE_SET_PITCH_SCALE,
+
+			// Emitter Commands
+			EMITTER_CREATE,
+			EMITTER_RELEASE,
+			EMITTER_SET_POSITION,
+
+		};
+	}
+
+	/**
+	* EMainThreadCommand
+	* Commands sent from audio system thread to main thread.
+	*/
+	namespace EMainThreadCommand
+	{
+		enum Type
+		{
+			NONE = 0,
+
+			// Voice Commands
+			VOICE_DONE,
+			VOICE_REAL,
+			VOICE_VIRTUAL,
+			VOICE_SUSPEND,
+		};
+	}
+
 	/** 
 	* FUnrealAudioModule 
 	* Concrete implementation of IUnrealAudioModule
 	*/
-	class FUnrealAudioModule : public IUnrealAudioModule
+	class FUnrealAudioModule :	public IUnrealAudioModule,
+								public FRunnable
 	{
 	public:
 
 		FUnrealAudioModule();
-		~FUnrealAudioModule();
+		~FUnrealAudioModule() override;
+
+		/************************************************************************/
+		/* IUnrealAudioModule													*/
+		/************************************************************************/
 
-		// IUnrealaAudioModule Interface
 		bool Initialize() override;
 		bool Initialize(const FString& DeviceModuleName) override;
+		void Update() override;
 		void Shutdown() override;
-		TSharedPtr<ISoundFile> ImportSound(const FSoundFileImportSettings& ImportSettings) override;
-		void ExportSound(TSharedPtr<ISoundFileData> SoundFileData, const FString& ExportPath) override;
+		int32 GetNumBackgroundTasks() const override;
+
+		// Sound file API
+		void ConvertSound(const FString& FilePath, const FString& OutputFilePath, const FSoundFileConvertFormat& ConvertFormat) override;
+		void ConvertSound(const FString& FilePath, const FString& OutputFilePath) override;
+
+		TSharedPtr<FSoundFileReader> CreateSoundFileReader();
+		TSharedPtr<FSoundFileWriter> CreateSoundFileWriter();
+
+		TSharedPtr<ISoundFile> LoadSoundFile(const FName& Name, TArray<uint8>& InBulkData) override;
+		TSharedPtr<ISoundFile> LoadSoundFile(const FName& Path, bool bLoadAsync = true) override;
+		TSharedPtr<ISoundFile> StreamSoundFile(const FName& Path, bool bLoadAsync = true) override;
+		
+		int32 GetNumSoundFilesStreamed() const override;
+		int32 GetNumSoundFilesLoaded() const override;
+		int32 GetSoundFileNumBytes() const override;
+		float GetSoundFilePercentageOfTargetMemoryLimit() const override;
+		void LogSoundFileMemoryInfo() const override;
+
+		// Creation functions
+		TSharedPtr<IEmitter> EmitterCreate() override;
+		TSharedPtr<IVoice> VoiceCreate(const FVoiceInitializationParams& Params) override;
+
+		class IUnrealAudioDeviceModule* GetDeviceModule() override;
+
+		/************************************************************************/
+		/* FRunnable implementation												*/
+		/************************************************************************/
+
+		bool Init();
+		uint32 Run();
+		void Stop();
+
+		/************************************************************************/
+		/* Internal Public API													*/
+		/************************************************************************/
 
-		class IUnrealAudioDeviceModule* GetDeviceModule();
 		FName GetDefaultDeviceModuleName() const;
 
-		// Internal API
+		bool AudioDeviceCallback(struct FCallbackInfo& CallbackInfo);
+
+		void GetPitchRange(float& OutMinPitch, float& OutMaxPitch) const;
 		void IncrementBackgroundTaskCount();
 		void DecrementBackgroundTaskCount();
+		FEmitterManager& GetEmitterManager();
+		FVoiceManager& GetVoiceManager();
+		FSoundFileManager& GetSoundFileManager();
+
+		void SendAudioThreadCommand(const FCommand& Command);
+		void SendMainThreadCommand(const FCommand& Command);
+
+		double GetCurrentTimeSec() const;
 
 	private:
 
 		bool InitializeInternal();
 		
-		/** Loads the sound file parsing libs */
+		bool InitializeAudioDevice();
+		bool InitializeAudioSystem();
+
+		void ShutdownAudioDevice();
+
+		void ExecuteAudioThreadCommands();
+		void ExecuteMainThreadCommands();
+
 		bool LoadSoundFileLib();
 		bool ShutdownSoundFileLib();
 
-		/** Initializes unreal audio tests */
-		static void InitializeTests(FUnrealAudioModule* Module);
+		/** Initializes unreal audio tests. */
+		bool DeviceTestCallback(struct FCallbackInfo& CallbackInfo);
+
+		static void InitializeSystemTests(IUnrealAudioModule* Module);
+		static void InitializeDeviceTests(IUnrealAudioModule* Module);
+		static void UpdateSystemTests();
 
 		class IUnrealAudioDeviceModule* UnrealAudioDevice;
+
 		FName ModuleName;
 
-		/** Number of background tasks in-flight */
+		/** The default settings for converting sound files. */
+		FSoundFileConvertFormat DefaultConvertFormat;
+
+		/** Number of background tasks in-flight. */
 		FThreadSafeCounter NumBackgroundTasks;
 
-		/** Handle to the sound file DLL */
+		/** Emitter manager object. */
+		FEmitterManager EmitterManager;
+
+		/** Voice manager object. */
+		FVoiceManager VoiceManager;
+
+		/** Sound file manager object. */
+		FSoundFileManager SoundFileManager;
+
+		/** Handle to the sound file DLL. */
 		void* SoundFileDllHandle;
 
+		/** Whether or not to stop the system thread. */
+		FThreadSafeBool bIsStoppingSystemThread;
+
+		/** The current time in seconds of the audio system thread. */
+		double AudioSystemTimeSec;
+
+		/** The audio system thread handle. */
+		FRunnableThread* SystemThread;
+
+		/** Command queue to send thread commands from main thread to audio thread. */
+		TCommandQueue<FCommand> AudioThreadCommandQueue;
 
+		/** Command queue to send thread commands from audio thread to main thread. */
+		TCommandQueue<FCommand> MainThreadCommandQueue;
+
+		double SystemThreadUpdateTime;
+
+		FThreadChecker MainThreadChecker;
+		FThreadChecker AudioThreadChecker;
 	};
 
+	/**
+	* Function called when an error occurs in the audio system code.
+	* @param Error The error type that occurred.
+	* @param ErrorDetails A string of specific details about the error.
+	* @param FileName The file that the error occurred in.
+	* @param LineNumber The line number that the error occurred in.
+	*/
+	static inline void OnSystemError(const ESystemError::Type Error, const FString& ErrorDetails, const FString& FileName, int32 LineNumber)
+	{
+		UE_LOG(LogUnrealAudio, Error, TEXT("Audio System Error: (%s) : %s (%s::%d)"), ESystemError::ToString(Error), *ErrorDetails, *FileName, LineNumber);
+	}
 
 
 }
diff --git a/Engine/Source/Runtime/UnrealAudio/Private/UnrealAudioSampleRateConverter.cpp b/Engine/Source/Runtime/UnrealAudio/Private/UnrealAudioSampleRateConverter.cpp
new file mode 100644
--- /dev/null
+++ b/Engine/Source/Runtime/UnrealAudio/Private/UnrealAudioSampleRateConverter.cpp
@@ -0,0 +1,174 @@
+
+#include "UnrealAudioPrivate.h"
+#include "UnrealAudioSampleRateConverter.h"
+
+#if ENABLE_UNREAL_AUDIO
+
+#define AUDIO_MIN_SAMPLE_RATE_RATIO_MAGNITUDE	(0.0001f)
+#define AUDIO_TARGET_SAMPLE_RATE_SMOOTHNESS		(0.001f)
+
+namespace UAudio
+{
+
+	FSampleRateConverter::FSampleRateConverter()
+		: NumChannels(0)
+		, PrevFrameIndex(0)
+		, NextFrameIndex(1)
+		, CurrentRateRatio(1.0f)
+		, TargetRateRatio(1.0f)
+		, CurrentFrameFraction(0.0f)
+		, bCachePrevValues(true)
+	{
+	}
+
+	void FSampleRateConverter::Init(float InRate, int32 InNumChannels)
+	{
+		NumChannels = InNumChannels;
+		PrevFrameIndex = 0;
+		NextFrameIndex = 0;
+		CurrentRateRatio = InRate;
+		TargetRateRatio = InRate;
+		CurrentFrameFraction = 0.0f;
+		bCachePrevValues = true;
+
+		DEBUG_AUDIO_CHECK(NumChannels > 0);
+		PrevFrameValues.Init(0.0f, NumChannels);
+		NextFrameValues.Init(0.0f, NumChannels);
+	}
+
+	void FSampleRateConverter::SetRateRatio(float InRate)
+	{
+		// We can't have a 0.0f rate ratio...
+		if (FMath::Abs(InRate) < AUDIO_MIN_SAMPLE_RATE_RATIO_MAGNITUDE)
+		{
+			if (InRate < 0.0f)
+			{
+				TargetRateRatio = -AUDIO_MIN_SAMPLE_RATE_RATIO_MAGNITUDE;
+			}
+			else
+			{
+				TargetRateRatio = AUDIO_MIN_SAMPLE_RATE_RATIO_MAGNITUDE;
+			}
+		}
+		else
+		{
+			TargetRateRatio = InRate;
+		}
+	}
+
+	float FSampleRateConverter::GetRateRatio() const
+	{
+		return CurrentRateRatio;
+	}
+
+	void FSampleRateConverter::ProcessBlock(const float* InputBuffer, int32 NumInputSamples, TArray<float>& OutputBuffer)
+	{
+		DEBUG_AUDIO_CHECK(TargetRateRatio != 0.0);
+
+		// If our input buffer is empty, then we have nothing to do
+		if (NumInputSamples == 0)
+		{
+			return;
+		}
+
+		// Our input buffer is interleaved samples, so frame is samples divided by channels
+		int32 InputFrames = NumInputSamples / NumChannels;
+
+		bool bFinished = false;
+		while (!bFinished)
+		{
+			if (CurrentRateRatio > 0.0)
+			{
+				ProcessForward(bFinished, InputBuffer, InputFrames, OutputBuffer);
+			}
+			else
+			{
+				ProcessBackward(bFinished, InputBuffer, InputFrames, OutputBuffer);
+			}
+		}
+	}
+
+	void FSampleRateConverter::ProcessForward(bool& bFinished, const float* InputBuffer, int32 InputFrames, TArray<float>& OutputBuffer)
+	{
+		// Process the input block until there's no more frames to process
+		while (CurrentRateRatio > 0.0f && PrevFrameIndex < InputFrames && NextFrameIndex < InputFrames)
+		{
+			DEBUG_AUDIO_CHECK(CurrentFrameFraction >= 0.0f && CurrentFrameFraction < 1.0f);
+
+			int32 PrevSampleIndex = NumChannels * PrevFrameIndex;
+			int32 NextSampleIndex = NumChannels * NextFrameIndex;
+
+			for (int32 Channel = 0; Channel < NumChannels; ++Channel)
+			{
+				if (bCachePrevValues)
+				{
+					DEBUG_AUDIO_CHECK(PrevSampleIndex >= 0);
+					PrevFrameValues[Channel] = InputBuffer[PrevSampleIndex + Channel];
+				}
+				DEBUG_AUDIO_CHECK(NextSampleIndex >= 0);
+				NextFrameValues[Channel] = InputBuffer[NextSampleIndex + Channel];
+			}
+
+			bCachePrevValues = true;
+
+			// while our current frame fraction is not another frame and our rate ratio is still positive
+			while (CurrentFrameFraction < 1.0f && CurrentRateRatio > 0.0f)
+			{
+				// Perform a sample-rate conversion on the input block 
+				for (int32 Channel = 0; Channel < NumChannels; ++Channel)
+				{
+					// linear interpolate
+					float ChannelValue = PrevFrameValues[Channel] * (1.0 - CurrentFrameFraction) + CurrentFrameFraction * NextFrameValues[Channel];
+					OutputBuffer.Add(ChannelValue);
+				}
+				CurrentFrameFraction += CurrentRateRatio;
+				if (CurrentRateRatio != TargetRateRatio)
+				{
+					CurrentRateRatio = CurrentRateRatio + AUDIO_TARGET_SAMPLE_RATE_SMOOTHNESS * (TargetRateRatio - CurrentRateRatio);
+				}
+			}
+
+			// Wrap the fraction to less than 1.0. This can be potentially multiples above 1.0
+			// for sample rates that significantly less than the input file (e.g. 8k vs 48k)
+			while (CurrentFrameFraction >= 1.0f)
+			{
+				++PrevFrameIndex;
+				++NextFrameIndex;
+				CurrentFrameFraction -= 1.0f;
+			}
+		}
+
+		// If we didn't change directions in this process block
+		if (CurrentRateRatio > 0.0f)
+		{
+			// we need to cache the last value of this buffer if we'll need it in the next block
+			if (PrevFrameIndex < InputFrames)
+			{
+				bCachePrevValues = false;
+				for (int32 Channel = 0; Channel < NumChannels; ++Channel)
+				{
+					PrevFrameValues[Channel] = NextFrameValues[Channel];
+				}
+			}
+
+			// update our frame index relative to the last InputFrame count
+			// Note: This could result in a negative PrevFrameIndex if PrevFrameIndex was less than InputFrames
+			PrevFrameIndex -= InputFrames;
+			NextFrameIndex = PrevFrameIndex + 1;
+
+			bFinished = true;
+		}
+	}
+
+	void FSampleRateConverter::ProcessBackward(bool& bFinished, const float* InputBuffer, int32 NumFrames, TArray<float>& OutputBuffer)
+	{
+		// TODO: do backward processing algorithm
+		bFinished = true;
+	}
+
+
+
+}
+
+#endif // #if ENABLE_UNREAL_AUDIO
+
diff --git a/Engine/Source/Runtime/UnrealAudio/Private/UnrealAudioSampleRateConverter.h b/Engine/Source/Runtime/UnrealAudio/Private/UnrealAudioSampleRateConverter.h
new file mode 100644
--- /dev/null
+++ b/Engine/Source/Runtime/UnrealAudio/Private/UnrealAudioSampleRateConverter.h
@@ -0,0 +1,44 @@
+
+#pragma once
+
+#include "UnrealAudioPrivate.h"
+
+#if ENABLE_UNREAL_AUDIO
+
+
+namespace UAudio
+{
+
+	class FSampleRateConverter
+	{
+	public:
+		FSampleRateConverter();
+
+		void Init(float InRate, int32 InNumChannels);
+
+		void SetRateRatio(float InRate);
+		float GetRateRatio() const;
+		void ProcessBlock(const float* InputBuffer, int32 NumInputSamples, TArray<float>& OutputBuffer);
+
+	private:
+		void ProcessForward(bool& bFinished, const float* InputBuffer, int32 InputFrames, TArray<float>& OutputBuffer);
+		void ProcessBackward(bool& bFinished, const float* InputBuffer, int32 InputFrames, TArray<float>& OutputBuffer);
+
+		TArray<float> PrevFrameValues;
+		TArray<float> NextFrameValues;
+
+		int32 NumChannels;
+
+		int32 PrevFrameIndex;
+		int32 NextFrameIndex;
+		float CurrentRateRatio;
+		float TargetRateRatio;
+		float CurrentFrameFraction;
+		bool bCachePrevValues;
+	};
+
+
+}
+
+#endif // #if ENABLE_UNREAL_AUDIO
+
