Unreal Audio WIP

diff --git a/Engine/Source/Runtime/UnrealAudio/Private/UnrealAudioSoundFile.cpp b/Engine/Source/Runtime/UnrealAudio/Private/UnrealAudioSoundFile.cpp
--- a/Engine/Source/Runtime/UnrealAudio/Private/UnrealAudioSoundFile.cpp
+++ b/Engine/Source/Runtime/UnrealAudio/Private/UnrealAudioSoundFile.cpp
@@ -8,909 +8,142 @@
 
 namespace UAudio
 {
-	// Virtual Sound File Function Pointers
-	typedef SoundFileCount(*VirtualSoundFileGetLengthFuncPtr)(void* UserData);
-	typedef SoundFileCount(*VirtualSoundFileSeekFuncPtr)(SoundFileCount Offset, int32 Mode, void* UserData);
-	typedef SoundFileCount(*VirtualSoundFileReadFuncPtr)(void* DataPtr, SoundFileCount ByteCount, void* UserData);
-	typedef SoundFileCount(*VirtualSoundFileWriteFuncPtr)(const void* DataPtr, SoundFileCount ByteCount, void* UserData);
-	typedef SoundFileCount(*VirtualSoundFileTellFuncPtr)(void* UserData);
-
-	// Struct describing function pointers to call for virtual file IO
-	struct FVirtualSoundFileCallbackInfo
-	{
-		VirtualSoundFileGetLengthFuncPtr VirtualSoundFileGetLength;
-		VirtualSoundFileSeekFuncPtr VirtualSoundFileSeek;
-		VirtualSoundFileReadFuncPtr VirtualSoundFileRead;
-		VirtualSoundFileWriteFuncPtr VirtualSoundFileWrite;
-		VirtualSoundFileTellFuncPtr VirtualSoundFileTell;
-	};
-
-	// SoundFile Constants
-	static const int32 SET_ENCODING_QUALITY = 0x1300;
-	static const int32 SET_CHANNEL_MAP_INFO = 0x1101;
-	static const int32 GET_CHANNEL_MAP_INFO = 0x1100;
-
-	// Exported SoundFile Functions
-	typedef SoundFileHandle*(*SoundFileOpenFuncPtr)(const char* Path, int32 Mode, FSoundFileDescription* Description);
-	typedef SoundFileHandle*(*SoundFileOpenVirtualFuncPtr)(FVirtualSoundFileCallbackInfo* VirtualFileDescription, int32 Mode, FSoundFileDescription* Description, void* UserData);
-	typedef int32(*SoundFileCloseFuncPtr)(SoundFileHandle* FileHandle);
-	typedef int32(*SoundFileErrorFuncPtr)(SoundFileHandle* FileHandle);
-	typedef const char*(*SoundFileStrErrorFuncPtr)(SoundFileHandle* FileHandle);
-	typedef const char*(*SoundFileErrorNumberFuncPtr)(int32 ErrorNumber);
-	typedef int32(*SoundFileCommandFuncPtr)(SoundFileHandle* FileHandle, int32 Command, void* Data, int32 DataSize);
-	typedef int32(*SoundFileFormatCheckFuncPtr)(const FSoundFileDescription* Description);
-	typedef SoundFileCount(*SoundFileSeekFuncPtr)(SoundFileHandle* FileHandle, SoundFileCount NumFrames, int32 SeekMode);
-	typedef const char*(*SoundFileGetVersionFuncPtr)(void);
-	typedef SoundFileCount(*SoundFileReadFramesFloatFuncPtr)(SoundFileHandle* FileHandle, float* Buffer, SoundFileCount NumFrames);
-	typedef SoundFileCount(*SoundFileReadFramesDoubleFuncPtr)(SoundFileHandle* FileHandle, double* Buffer, SoundFileCount NumFrames);
-	typedef SoundFileCount(*SoundFileWriteFramesFloatFuncPtr)(SoundFileHandle* FileHandle, const float* Buffer, SoundFileCount NumFrames);
-	typedef SoundFileCount(*SoundFileWriteFramesDoubleFuncPtr)(SoundFileHandle* FileHandle, const double* Buffer, SoundFileCount NumFrames);
-	typedef SoundFileCount(*SoundFileReadSamplesFloatFuncPtr)(SoundFileHandle* FileHandle, float* Buffer, SoundFileCount NumSamples);
-	typedef SoundFileCount(*SoundFileReadSamplesDoubleFuncPtr)(SoundFileHandle* FileHandle, double* Buffer, SoundFileCount NumSamples);
-	typedef SoundFileCount(*SoundFileWriteSamplesFloatFuncPtr)(SoundFileHandle* FileHandle, const float* Buffer, SoundFileCount NumSamples);
-	typedef SoundFileCount(*SoundFileWriteSamplesDoubleFuncPtr)(SoundFileHandle* FileHandle, const double* Buffer, SoundFileCount NumSamples);
-
-	SoundFileOpenFuncPtr SoundFileOpen = nullptr;
-	SoundFileOpenVirtualFuncPtr SoundFileOpenVirtual = nullptr;
-	SoundFileCloseFuncPtr SoundFileClose = nullptr;
-	SoundFileErrorFuncPtr SoundFileError = nullptr;
-	SoundFileStrErrorFuncPtr SoundFileStrError = nullptr;
-	SoundFileErrorNumberFuncPtr SoundFileErrorNumber = nullptr;
-	SoundFileCommandFuncPtr SoundFileCommand = nullptr;
-	SoundFileFormatCheckFuncPtr SoundFileFormatCheck = nullptr;
-	SoundFileSeekFuncPtr SoundFileSeek = nullptr;
-	SoundFileGetVersionFuncPtr SoundFileGetVersion = nullptr;
-	SoundFileReadFramesFloatFuncPtr SoundFileReadFramesFloat = nullptr;
-	SoundFileReadFramesDoubleFuncPtr SoundFileReadFramesDouble = nullptr;
-	SoundFileWriteFramesFloatFuncPtr SoundFileWriteFramesFloat = nullptr;
-	SoundFileWriteFramesDoubleFuncPtr SoundFileWriteFramesDouble = nullptr;
-	SoundFileReadSamplesFloatFuncPtr SoundFileReadSamplesFloat = nullptr;
-	SoundFileReadSamplesDoubleFuncPtr SoundFileReadSamplesDouble = nullptr;
-	SoundFileWriteSamplesFloatFuncPtr SoundFileWriteSamplesFloat = nullptr;
-	SoundFileWriteSamplesDoubleFuncPtr SoundFileWriteSamplesDouble = nullptr;
-
-	/** 
-	Function implementations of virtual function callbacks
-	*/
-
-	static SoundFileCount OnSoundFileGetLengthBytes(void* UserData)
-	{
-		SoundFileCount Length = 0;
-		((ISoundFileVirtual*)UserData)->GetLengthBytes(Length);
-		return Length;
-	}
-
-	static SoundFileCount OnSoundFileSeekBytes(SoundFileCount Offset, int32 Mode, void* UserData)
-	{
-		SoundFileCount OutOffset = 0;
-		((ISoundFileVirtual*)UserData)->SeekBytes(Offset, (ESoundFileSeekMode::Type)Mode, OutOffset);
-		return OutOffset;
-	}
-
-	static SoundFileCount OnSoundFileReadBytes(void* DataPtr, SoundFileCount ByteCount, void* UserData)
-	{
-		SoundFileCount OutBytesRead = 0;
-		((ISoundFileVirtual*)UserData)->ReadBytes(DataPtr, ByteCount, OutBytesRead);
-		return OutBytesRead;
-	}
-
-	static SoundFileCount OnSoundFileWriteBytes(const void* DataPtr, SoundFileCount ByteCount, void* UserData)
-	{
-		SoundFileCount OutBytesWritten = 0;
-		((ISoundFileVirtual*)UserData)->WriteBytes(DataPtr, ByteCount, OutBytesWritten);
-		return OutBytesWritten;
-	}
-
-	static SoundFileCount OnSoundFileTell(void* UserData)
-	{
-		return ((ISoundFileVirtual*)UserData)->GetOffsetBytes();
-	}
-
-	/**
-	* Gets the default channel mapping for the given channel number
-	*/
-	static void GetDefaultMappingsForChannelNumber(int32 NumChannels, TArray<ESoundFileChannelMap::Type>& ChannelMap)
-	{
-		check(ChannelMap.Num() == NumChannels);
-
-		switch (NumChannels)
-		{
-			case 1:	// MONO
-				ChannelMap[0] = ESoundFileChannelMap::MONO;
-				break;
-
-			case 2:	// STEREO
-				ChannelMap[0] = ESoundFileChannelMap::LEFT;
-				ChannelMap[1] = ESoundFileChannelMap::RIGHT;
-				break;
-
-			case 3:	// 2.1
-				ChannelMap[0] = ESoundFileChannelMap::LEFT;
-				ChannelMap[1] = ESoundFileChannelMap::RIGHT;
-				ChannelMap[2] = ESoundFileChannelMap::LFE;
-				break;
-
-			case 4: // Quadraphonic
-				ChannelMap[0] = ESoundFileChannelMap::LEFT;
-				ChannelMap[1] = ESoundFileChannelMap::RIGHT;
-				ChannelMap[2] = ESoundFileChannelMap::BACK_LEFT;
-				ChannelMap[3] = ESoundFileChannelMap::BACK_RIGHT;
-				break;
-
-			case 5: // 5.0
-				ChannelMap[0] = ESoundFileChannelMap::LEFT;
-				ChannelMap[1] = ESoundFileChannelMap::RIGHT;
-				ChannelMap[2] = ESoundFileChannelMap::CENTER;
-				ChannelMap[3] = ESoundFileChannelMap::SIDE_LEFT;
-				ChannelMap[4] = ESoundFileChannelMap::SIDE_RIGHT;
-				break;
-
-			case 6: // 5.1
-				ChannelMap[0] = ESoundFileChannelMap::LEFT;
-				ChannelMap[1] = ESoundFileChannelMap::RIGHT;
-				ChannelMap[2] = ESoundFileChannelMap::CENTER;
-				ChannelMap[3] = ESoundFileChannelMap::LFE;
-				ChannelMap[4] = ESoundFileChannelMap::SIDE_LEFT;
-				ChannelMap[5] = ESoundFileChannelMap::SIDE_RIGHT;
-				break;
-
-			case 7: // 6.1
-				ChannelMap[0] = ESoundFileChannelMap::LEFT;
-				ChannelMap[1] = ESoundFileChannelMap::RIGHT;
-				ChannelMap[2] = ESoundFileChannelMap::CENTER;
-				ChannelMap[3] = ESoundFileChannelMap::LFE;
-				ChannelMap[4] = ESoundFileChannelMap::SIDE_LEFT;
-				ChannelMap[5] = ESoundFileChannelMap::SIDE_RIGHT;
-				ChannelMap[6] = ESoundFileChannelMap::BACK_CENTER;
-				break;
-
-			case 8: // 7.1
-				ChannelMap[0] = ESoundFileChannelMap::LEFT;
-				ChannelMap[1] = ESoundFileChannelMap::RIGHT;
-				ChannelMap[2] = ESoundFileChannelMap::CENTER;
-				ChannelMap[3] = ESoundFileChannelMap::LFE;
-				ChannelMap[4] = ESoundFileChannelMap::BACK_LEFT;
-				ChannelMap[5] = ESoundFileChannelMap::BACK_RIGHT;
-				ChannelMap[6] = ESoundFileChannelMap::SIDE_LEFT;
-				ChannelMap[7] = ESoundFileChannelMap::SIDE_RIGHT;
-				break;
-
-			default:
-				break;
-		}
-	}
-
-	static int32 SoundFileId = 0;
-
-	FSoundFile::FSoundFile()
-		: CurrentIndexBytes(0)
-		, FileHandle(nullptr)
-		, State(ESoundFileState::UNINITIALIZED)
-		, CurrentError(ESoundFileError::NONE)
-		, Id(SoundFileId++)
+	FSoundFile::FSoundFile(FSoundFileManager* InSoundFileManager)
+		: SoundFileManager(InSoundFileManager)
+		, Error(ESoundFileError::NONE)
 	{
+		check(SoundFileManager);
 	}
 
 	FSoundFile::~FSoundFile()
 	{
-		check(FileHandle == nullptr);
-	}
-
-	ESoundFileError::Type FSoundFile::GetSoundFileData(TSharedPtr<ISoundFileData>& OutSoundFileData)
-	{
-		if (State == ESoundFileState::IMPORTING)
-		{
-			return SetError(ESoundFileError::ALREADY_INTIALIZED);
-		}
-
-		OutSoundFileData = SoundFileData;
-		return ESoundFileError::NONE;
+		SoundFileManager->ReleaseSoundFileHandle(SoundFileHandle);
 	}
 
-	ESoundFileError::Type FSoundFile::Initialize(TSharedPtr<ISoundFileData> InSoundFileData)
+	ESoundFileError::Type FSoundFile::GetState(ESoundFileState::Type& OutState) const
 	{
-		if (State != ESoundFileState::UNINITIALIZED)
-		{
-			return SetError(ESoundFileError::ALREADY_INTIALIZED);
-		}
-
-		check(FileHandle == nullptr);
-
-		// Setting sound file data initializes this sound file
-		SoundFileData = InSoundFileData;
-
-		// Open up a virtual file handle with this data
-		FVirtualSoundFileCallbackInfo VirtualSoundFileInfo;
-		VirtualSoundFileInfo.VirtualSoundFileGetLength = OnSoundFileGetLengthBytes;
-		VirtualSoundFileInfo.VirtualSoundFileSeek = OnSoundFileSeekBytes;
-		VirtualSoundFileInfo.VirtualSoundFileRead = OnSoundFileReadBytes;
-		VirtualSoundFileInfo.VirtualSoundFileWrite = OnSoundFileWriteBytes;
-		VirtualSoundFileInfo.VirtualSoundFileTell = OnSoundFileTell;
-
-		FSoundFileDescription Description = SoundFileData->GetDescription();
-		if (!SoundFileFormatCheck(&Description))
-		{
-			return SetError(ESoundFileError::INVALID_INPUT_FORMAT);
-		}
-
-		FileHandle = SoundFileOpenVirtual(&VirtualSoundFileInfo, ESoundFileOpenMode::READING, &Description, (void*)this);
-		if (!FileHandle)
-		{
-			FString StrErr = SoundFileStrError(nullptr);
-			UE_LOG(LogUnrealAudio, Error, TEXT("Failed to intitialize sound file: %s"), *StrErr);
-			return SetError(ESoundFileError::FAILED_TO_OPEN);
-		}
-
-		State = ESoundFileState::INITIALIZED;
+		OutState = SoundFileManager->GetState(SoundFileHandle);
 		return ESoundFileError::NONE;
 	}
 
-	ESoundFileError::Type FSoundFile::GetName(FString& Name)
-	{
-		if (SoundFileData.IsValid())
-		{
-			const FString& FilePath = SoundFileData->GetFilePath();
-			Name = FPaths::GetBaseFilename(FilePath);
-			return ESoundFileError::NONE;
-		}
-		return ESoundFileError::INVALID_DATA;
-	}
-
-	ESoundFileState::Type FSoundFile::GetState() const
-	{
-		return State;
-	}
-
 	ESoundFileError::Type FSoundFile::GetError() const
 	{
-		return CurrentError;
+		return Error;
 	}
 
-	ESoundFileError::Type FSoundFile::SetError(ESoundFileError::Type Error)
+	ESoundFileError::Type FSoundFile::GetId(uint32& OutId) const
 	{
-		if (Error != ESoundFileError::NONE)
-		{
-			State = ESoundFileState::HAS_ERROR;
-		}
-		return (CurrentError = Error);
-	}
-
-	ESoundFileError::Type FSoundFile::OpenFileForReading(const FString& FilePath)
-	{
-		if (FileHandle != nullptr)
-		{
-			UE_LOG(LogUnrealAudio, Error, TEXT("Sound file already opened"));
-			return SetError(ESoundFileError::ALREADY_OPENED);
-		}
-
-		// Make sure we don't already have data
-		if (SoundFileData.IsValid())
-		{
-			return SetError(ESoundFileError::ALREADY_HAS_DATA);
-		}
-
-		// Make a data ptr
-		SoundFileData = TSharedPtr<FSoundFileData>(new FSoundFileData(ESoundFileOpenMode::READING));
-		SoundFileData->SetFilePath(FilePath);
-
-		// Check to see if the file exists
-		if (!FPaths::FileExists(FilePath))
+		if (!SoundFileManager->EntityManager.IsValidEntity(SoundFileHandle))
 		{
-			UE_LOG(LogUnrealAudio, Error, TEXT("Sound file %s doesn't exist."), *FilePath);
-			return SetError(ESoundFileError::FILE_DOESNT_EXIST);
+			return ESoundFileError::INVALID_SOUND_FILE_HANDLE;
 		}
-
-		// Open up the file
-		FSoundFileDescription Description;
-		FileHandle = SoundFileOpen(TCHAR_TO_ANSI(*FilePath), ESoundFileOpenMode::READING, &Description);
-		if (!FileHandle)
-		{
-			FString StrError = FString(SoundFileStrError(nullptr));
-			UE_LOG(LogUnrealAudio, Error, TEXT("Failed to open sound file %s: %s"), *FilePath, *StrError);
-			return SetError(ESoundFileError::FAILED_TO_OPEN);
-		}
-		SoundFileData->SetDescription(Description);
-
-		// Try to get a channel mapping
-		int32 NumChannels = Description.NumChannels;
-		TArray<ESoundFileChannelMap::Type> ChannelMap;
-		ChannelMap.Init(ESoundFileChannelMap::INVALID, NumChannels);
-
-		int32 Result = SoundFileCommand(FileHandle, GET_CHANNEL_MAP_INFO, (int32*)ChannelMap.GetData(), sizeof(int32)*NumChannels);
-
-		// If we failed to get the file's channel map definition, then we set the default based on the number of channels
-		if (Result == 0)
-		{
-			GetDefaultMappingsForChannelNumber(NumChannels, ChannelMap);
-		}
-		else
-		{
-			// Check to see if the channel map we did get back is filled with INVALID channels
-			bool bIsInvalid = false;
-			for (ESoundFileChannelMap::Type ChannelType : ChannelMap)
-			{
-				if (ChannelType == ESoundFileChannelMap::INVALID)
-				{
-					bIsInvalid = true;
-					break;
-				}
-			}
-			// If invalid, then we need to get the default channel mapping
-			if (bIsInvalid)
-			{
-				GetDefaultMappingsForChannelNumber(NumChannels, ChannelMap);
-			}
-		}
-		SoundFileData->SetChannelMap(ChannelMap);
-
+		OutId = SoundFileHandle.Id;
 		return ESoundFileError::NONE;
 	}
 
-	ESoundFileError::Type FSoundFile::OpenEmptyFileForImport(const FSoundFileDescription& InDescription, const TArray<ESoundFileChannelMap::Type>& InChannelMap)
+	ESoundFileError::Type FSoundFile::Init(const FSoundFileHandle& InSoundFileHandle)
 	{
-		// First check the input format to make sure it's valid
-		if (!SoundFileFormatCheck(&InDescription))
-		{
-			UE_LOG(LogUnrealAudio, 
-				   Error, 
-				   TEXT("Sound file input format (%s - %s) is invalid."), 
-				   ESoundFileFormat::ToStringMajor(InDescription.FormatFlags), 
-				   ESoundFileFormat::ToStringMinor(InDescription.FormatFlags));
-
-			return SetError(ESoundFileError::INVALID_INPUT_FORMAT);
-		}
-		
-		// Make sure we don't already have data
-		if (SoundFileData.IsValid())
-		{
-			return SetError(ESoundFileError::ALREADY_HAS_DATA);
-		}
-		
-		// Make sure we have the right number of channels and our channel map size
-		if (InChannelMap.Num() != InDescription.NumChannels)
+		if (!SoundFileManager->EntityManager.IsValidEntity(InSoundFileHandle))
 		{
-			UE_LOG(LogUnrealAudio, Error, TEXT("Channel map didn't match the input NumChannels"));
-			return (CurrentError = ESoundFileError::INVALID_CHANNEL_MAP);
+			return ESoundFileError::INVALID_SOUND_FILE_HANDLE;
 		}
-		
-		// Make a data ptr
-		SoundFileData = TSharedPtr<FSoundFileData>(new FSoundFileData(ESoundFileOpenMode::WRITING));
-
-		FVirtualSoundFileCallbackInfo VirtualSoundFileInfo;
-		VirtualSoundFileInfo.VirtualSoundFileGetLength = OnSoundFileGetLengthBytes;
-		VirtualSoundFileInfo.VirtualSoundFileSeek = OnSoundFileSeekBytes;
-		VirtualSoundFileInfo.VirtualSoundFileRead = OnSoundFileReadBytes;
-		VirtualSoundFileInfo.VirtualSoundFileWrite = OnSoundFileWriteBytes;
-		VirtualSoundFileInfo.VirtualSoundFileTell = OnSoundFileTell;
-
-		FSoundFileDescription Description = InDescription;
-		FileHandle = SoundFileOpenVirtual(&VirtualSoundFileInfo, ESoundFileOpenMode::WRITING, &Description, (void*)this);
-		if (!FileHandle)
-		{
-			FString StrErr = FString(SoundFileStrError(nullptr));
-			UE_LOG(LogUnrealAudio, Error, TEXT("Failed to open empty sound file: %s"), *StrErr);
-			return SetError(ESoundFileError::FAILED_TO_OPEN);
-		}
-
-		// For some reason, virtual sound files set the NumFrames to be the max int64 value, so we're going to set it back to 0
-		// since at this point, the number of frames is 0!
-		SoundFileData->SetDescription(Description);
-
-		// Set the channel map on the virtual file
-		SoundFileData->SetChannelMap(InChannelMap);
-		int32 Result = SoundFileCommand(FileHandle, SET_CHANNEL_MAP_INFO, (int32*)InChannelMap.GetData(), sizeof(int32)*Description.NumChannels);
-		if (Result)
-		{
-			FString StrErr = SoundFileStrError(nullptr);
-			UE_LOG(LogUnrealAudio, Error, TEXT("Failed to set the channel map on empty file for writing: %s"), *StrErr);
-			return SetError(ESoundFileError::INVALID_CHANNEL_MAP);
-		}
-
+		SoundFileHandle = InSoundFileHandle;
 		return ESoundFileError::NONE;
 	}
 
-	ESoundFileError::Type FSoundFile::OpenEmptyFileForExport(const FString& FilePath, const FSoundFileDescription& InDescription, const TArray<ESoundFileChannelMap::Type>& InChannelMap)
+	ESoundFileError::Type FSoundFile::SetError(ESoundFileError::Type InError)
 	{
-		// First check the input format to make sure it's valid
-		if (!SoundFileFormatCheck(&InDescription))
-		{
-			UE_LOG(LogUnrealAudio,
-				   Error,
-				   TEXT("Sound file input format (%s - %s) is invalid."),
-				   ESoundFileFormat::ToStringMajor(InDescription.FormatFlags),
-				   ESoundFileFormat::ToStringMinor(InDescription.FormatFlags));
-
-			return SetError(ESoundFileError::INVALID_INPUT_FORMAT);
-		}
-
-		// Make sure we have the right number of channels and our channel map size
-		if (InChannelMap.Num() != InDescription.NumChannels)
-		{
-			UE_LOG(LogUnrealAudio, Error, TEXT("Channel map didn't match the input NumChannels"));
-			return SetError(ESoundFileError::INVALID_CHANNEL_MAP);
-		}
-
-		// Open up the file
-		FSoundFileDescription Description = InDescription;
-		FileHandle = SoundFileOpen(TCHAR_TO_ANSI(*FilePath), ESoundFileOpenMode::WRITING, &Description);
-		if (!FileHandle)
-		{
-			FString StrErr = SoundFileStrError(nullptr);
-			UE_LOG(LogUnrealAudio, Error, TEXT("Failed to open sound file %s: %s"), *FilePath, *StrErr);
-			return SetError(ESoundFileError::FAILED_TO_OPEN);
-		}
-
-		return ESoundFileError::NONE;
-	}
-
-
-	ESoundFileError::Type FSoundFile::ReleaseSoundFileHandle()
-	{
-		if (FileHandle)
-		{
-			SoundFileClose(FileHandle);
-			FileHandle = nullptr;
-		}
-		return ESoundFileError::NONE;
-	}
-
-	ESoundFileError::Type FSoundFile::SetImportFileInfo(const FSoundFileDescription& FileDescription, const FString& FullPath)
-	{
-		if (!SoundFileData.IsValid())
-		{
-			return SetError(ESoundFileError::INVALID_DATA);
-		}
-
-		SoundFileData->SetOriginalDescription(FileDescription);
-		SoundFileData->SetFilePath(FullPath);
-		return ESoundFileError::NONE;
-	}
-
-	ESoundFileError::Type FSoundFile::SetEncodingQuality(double EncodingQuality)
-	{
-		if (FileHandle == nullptr || State != ESoundFileState::UNINITIALIZED)
-		{
-			return SetError(ESoundFileError::INVALID_SOUND_FILE);
-		}
-
-		if (!SoundFileData.IsValid())
-		{
-			return SetError(ESoundFileError::INVALID_DATA);
-		}
-
-		// This only does anything if the format is OGG
-		if ((SoundFileData->GetDescription().FormatFlags & ESoundFileFormat::MAJOR_FORMAT_MASK) == ESoundFileFormat::OGG)
-		{
-
-			int32 Result = SoundFileCommand(FileHandle, SET_ENCODING_QUALITY, &EncodingQuality, sizeof(double));
-			if (Result != 1)
-			{
-				FString StrErr = SoundFileStrError(FileHandle);
-				UE_LOG(LogUnrealAudio, Error, TEXT("Failed to set encoding quality: %s"), *StrErr);
-				return CurrentError = ESoundFileError::BAD_ENCODING_QUALITY;
-			}
-		}
-		return ESoundFileError::NONE;
-	}
-
-	ESoundFileError::Type FSoundFile::BeginImport()
-	{
-		State = ESoundFileState::IMPORTING;
-		return ESoundFileError::NONE;
-	}
-
-	ESoundFileError::Type FSoundFile::EndImport()
-	{
-		// Set the current NumFrames to the original NumFrames
-		((FSoundFileData*)SoundFileData.Get())->SetNumFrames(SoundFileData->GetOriginalDescription().NumFrames);
-		State = ESoundFileState::IMPORTED;
-		return ESoundFileError::NONE;
+		Error = InError;
+		return Error;
 	}
 
-	ESoundFileError::Type FSoundFile::SeekFrames(SoundFileCount Offset, ESoundFileSeekMode::Type SeekMode, SoundFileCount& OutOffset)
+	ESoundFileError::Type FSoundFile::GetPath(FName& OutPath) const
 	{
-		SoundFileCount Pos = SoundFileSeek(FileHandle, Offset, (int32)SeekMode);
-		if (Pos == -1)
+		TSharedPtr<FSoundFileDataEntry> Entry = SoundFileManager->GetEntry(SoundFileHandle);
+		if (Entry.IsValid())
 		{
-			FString StrErr = SoundFileStrError(FileHandle);
-			UE_LOG(LogUnrealAudio, Error, TEXT("Failed to seek file: %s"), *StrErr);
-			return SetError(ESoundFileError::FAILED_TO_SEEK);
-		}
-		return ESoundFileError::NONE;
-	}
-
-	ESoundFileError::Type FSoundFile::ReadFrames(float* DataPtr, SoundFileCount NumFrames, SoundFileCount& NumFramesRead)
-	{
-		NumFramesRead = SoundFileReadFramesFloat(FileHandle, DataPtr, NumFrames);
-		return ESoundFileError::NONE;
-	}
-
-	ESoundFileError::Type FSoundFile::ReadFrames(double* DataPtr, SoundFileCount NumFrames, SoundFileCount& NumFramesRead)
-	{
-		NumFramesRead = SoundFileReadFramesDouble(FileHandle, DataPtr, NumFrames);
-		return ESoundFileError::NONE;
-	}
-
-	ESoundFileError::Type FSoundFile::ReadSamples(float* DataPtr, SoundFileCount NumSamples, SoundFileCount& OutNumSamplesRead)
-	{
-		OutNumSamplesRead = SoundFileReadSamplesFloat(FileHandle, DataPtr, NumSamples);
-		return ESoundFileError::NONE;
-	}
-
-	ESoundFileError::Type FSoundFile::ReadSamples(double* DataPtr, SoundFileCount NumSamples, SoundFileCount& OutNumSamplesRead)
-	{
-		OutNumSamplesRead = SoundFileReadSamplesDouble(FileHandle, DataPtr, NumSamples);
-		return ESoundFileError::NONE;
-	}
-
-	ESoundFileError::Type FSoundFile::WriteFrames(const float* DataPtr, SoundFileCount NumFrames, SoundFileCount& OutNumFramesWritten)
-	{
-		OutNumFramesWritten = SoundFileWriteFramesFloat(FileHandle, DataPtr, NumFrames);
-		return ESoundFileError::NONE;
-	}
-
-	ESoundFileError::Type FSoundFile::WriteFrames(const double* DataPtr, SoundFileCount NumFrames, SoundFileCount& OutNumFramesWritten)
-	{
-		OutNumFramesWritten = SoundFileWriteFramesDouble(FileHandle, DataPtr, NumFrames);
-		return ESoundFileError::NONE;
-	}
-
-	ESoundFileError::Type FSoundFile::WriteSamples(const float* DataPtr, SoundFileCount NumSamples, SoundFileCount& OutNumSampleWritten)
-	{
-		OutNumSampleWritten = SoundFileWriteSamplesFloat(FileHandle, DataPtr, NumSamples);
-		return ESoundFileError::NONE;
-	}
-
-	ESoundFileError::Type FSoundFile::WriteSamples(const double* DataPtr, SoundFileCount NumSamples, SoundFileCount& OutNumSampleWritten)
-	{
-		OutNumSampleWritten = SoundFileWriteSamplesDouble(FileHandle, DataPtr, NumSamples);
-		return ESoundFileError::NONE;
-	}
-
-	ESoundFileError::Type FSoundFile::GetLengthBytes(SoundFileCount& OutLength) const
-	{
-		if (!SoundFileData.IsValid())
-		{
-			return ESoundFileError::INVALID_DATA;
+			OutPath = Entry->SoundFilePath;
+			return ESoundFileError::NONE;
 		}
-
-		OutLength = (SoundFileCount)SoundFileData->GetDataSize();
-		return ESoundFileError::NONE;
+		return ESoundFileError::INVALID_DATA;
 	}
 
-	ESoundFileError::Type FSoundFile::SeekBytes(SoundFileCount Offset, ESoundFileSeekMode::Type SeekMode, SoundFileCount& OutOffset)
+	ESoundFileError::Type FSoundFile::GetBulkData(TArray<uint8>** OutData) const
 	{
-		if (!SoundFileData.IsValid())
-		{
-			return ESoundFileError::INVALID_DATA;
-		}
-
-		// If there's no data, then seek does nothing
-		SoundFileCount MaxBytes = SoundFileData->GetDataSize();
-		if (MaxBytes == 0)
+		TSharedPtr<FSoundFileDataEntry> Entry = SoundFileManager->GetEntry(SoundFileHandle);
+		if (Entry.IsValid())
 		{
-			OutOffset = 0;
-			CurrentIndexBytes = 0;
+			*OutData = &Entry->BulkData;
 			return ESoundFileError::NONE;
 		}
-
-		switch (SeekMode)
-		{
-			case ESoundFileSeekMode::FROM_START:
-			CurrentIndexBytes = Offset;
-			break;
-
-			case ESoundFileSeekMode::FROM_CURRENT:
-			CurrentIndexBytes += Offset;
-			break;
-
-			case ESoundFileSeekMode::FROM_END:
-			CurrentIndexBytes = MaxBytes + Offset;
-			break;
-
-			default:
-			checkf(false, TEXT("Uknown seek mode!"));
-			break;
-		}
-
-		// Wrap the byte index to fall between 0 and MaxBytes
-		while (CurrentIndexBytes < 0)
-		{
-			CurrentIndexBytes += MaxBytes;
-		}
-
-		while (CurrentIndexBytes >= MaxBytes)
-		{
-			CurrentIndexBytes -= MaxBytes;
-		}
-
-		OutOffset = CurrentIndexBytes;
-		return ESoundFileError::NONE;
+		return ESoundFileError::INVALID_DATA;
 	}
 
-	ESoundFileError::Type FSoundFile::ReadBytes(void* DataPtr, SoundFileCount NumBytes, SoundFileCount& NumBytesRead)
+	ESoundFileError::Type FSoundFile::GetDataSize(int32& OutDataSize) const
 	{
-		if (!SoundFileData.IsValid())
-		{
-			return ESoundFileError::INVALID_DATA;
-		}
-
-		SoundFileCount EndByte = CurrentIndexBytes + NumBytes;
-		SoundFileCount MaxBytes = (SoundFileCount)SoundFileData->GetDataSize();
-		if (EndByte > MaxBytes)
+		TSharedPtr<FSoundFileDataEntry> Entry = SoundFileManager->GetEntry(SoundFileHandle);
+		if (Entry.IsValid())
 		{
-			NumBytes = MaxBytes - CurrentIndexBytes;
-		}
-
-		if (NumBytes > 0)
-		{
-			const uint8* DataArrayData = SoundFileData->GetData();
-			FMemory::Memcpy(DataPtr, (const void*)&DataArrayData[CurrentIndexBytes], NumBytes);
- 			CurrentIndexBytes += NumBytes;
+			OutDataSize = Entry->BulkData.Num();
+			return ESoundFileError::NONE;
 		}
-		NumBytesRead = NumBytes;
-		return ESoundFileError::NONE;
+		return ESoundFileError::INVALID_DATA;
 	}
 
-	ESoundFileError::Type FSoundFile::WriteBytes(const void* DataPtr, SoundFileCount NumBytes, SoundFileCount& NumBytesWritten)
+	ESoundFileError::Type FSoundFile::GetDescription(FSoundFileDescription& OutDescription) const
 	{
-		const uint8* InDataBytes = (const uint8*)DataPtr;
-		FSoundFileData* Data = (FSoundFileData*)SoundFileData.Get();
-		SoundFileCount MaxBytes = Data->GetDataArray().Num();
-
-		// Append the data to the buffer
-		if (CurrentIndexBytes == MaxBytes)
+		TSharedPtr<FSoundFileDataEntry> Entry = SoundFileManager->GetEntry(SoundFileHandle);
+		if (Entry.IsValid())
 		{
-			Data->GetDataArray().Append(InDataBytes, NumBytes);
-			CurrentIndexBytes += NumBytes;
-		}
-		else if ((CurrentIndexBytes + NumBytes) < MaxBytes)
-		{
-			// Write over the existing data in the buffer
-			for (SoundFileCount i = 0; i < NumBytes; ++i)
-			{
-				Data->GetDataArray()[CurrentIndexBytes++] = InDataBytes[i];
-			}
-		}
-		else
-		{
-			// Overwrite some data until end of current buffer size
-			SoundFileCount RemainingBytes = MaxBytes - CurrentIndexBytes;
-			SoundFileCount InputBufferIndex = 0;
-			for (InputBufferIndex = 0; InputBufferIndex < RemainingBytes; ++InputBufferIndex)
-			{
-				Data->GetDataArray()[CurrentIndexBytes++] = InDataBytes[InputBufferIndex];
-			}
-
-			// Now append the remainder of the input buffer to the internal data buffer
-			const uint8* RemainingData = &InDataBytes[InputBufferIndex];
-			RemainingBytes = NumBytes - RemainingBytes;
-			Data->GetDataArray().Append(RemainingData, RemainingBytes);
-			CurrentIndexBytes += RemainingBytes;
-			check(CurrentIndexBytes == Data->GetDataArray().Num());
+			OutDescription = Entry->Description;
+			return ESoundFileError::NONE;
 		}
-		NumBytesWritten = NumBytes;
-		return ESoundFileError::NONE;
-	}
-
-	SoundFileCount FSoundFile::GetOffsetBytes() const
-	{
-		return CurrentIndexBytes;
-	}
-
-	FSoundFileData::FSoundFileData(ESoundFileOpenMode::Type InMode)
-		: DataBuffer(nullptr)
-		, DataBufferSize(0)
-		, Mode(InMode)
-	{
-		Description = { 0 };
-		OriginalDescription = { 0 };
-	}
-
-	FSoundFileData::~FSoundFileData()
-	{
-	}
-
-	const FString& FSoundFileData::GetFilePath() const
-	{
-		return FilePath;
-	}
-
-	void FSoundFileData::SetFilePath(const FString& InPath)
-	{
-		FilePath = InPath;
-	}
-
-	const FSoundFileDescription& FSoundFileData::GetDescription() const
-	{
-		return Description;
-	}
-
-	void FSoundFileData::SetDescription(const FSoundFileDescription& InDescription)
-	{
-		Description = InDescription;
-	}
-
-	const FSoundFileDescription& FSoundFileData::GetOriginalDescription() const
-	{
-		return OriginalDescription;
-	}
-
-	void FSoundFileData::SetOriginalDescription(const FSoundFileDescription& InDescription)
-	{
-		OriginalDescription = InDescription;
-	}
-
-	const TArray<ESoundFileChannelMap::Type>& FSoundFileData::GetChannelMap() const
-	{
-		return ChannelMap;
-	}
-
-	void FSoundFileData::SetChannelMap(const TArray<ESoundFileChannelMap::Type>& InChannelMap)
-	{
-		ChannelMap = InChannelMap;
-	}
-
-	void FSoundFileData::SetNumFrames(int64 InNumFrames)
-	{
-		Description.NumFrames = InNumFrames;
+		return ESoundFileError::INVALID_DATA;
 	}
 
-	const uint8* FSoundFileData::GetData() const
+	ESoundFileError::Type FSoundFile::GetChannelMap(TArray<ESoundFileChannelMap::Type>& OutChannelMap) const
 	{
-		if (Mode == ESoundFileOpenMode::WRITING)
+		TSharedPtr<FSoundFileDataEntry> Entry = SoundFileManager->GetEntry(SoundFileHandle);
+		if (Entry.IsValid())
 		{
-			return DataArray.GetData();
-		}
-		else
-		{
-			return DataBuffer;
+			OutChannelMap = Entry->ChannelMap;
+			return ESoundFileError::NONE;
 		}
+		return ESoundFileError::INVALID_DATA;
 	}
 
-	int64 FSoundFileData::GetDataSize() const
+	ESoundFileError::Type FSoundFile::IsStreamed(bool& bOutIsStreamed) const
 	{
-		if (Mode == ESoundFileOpenMode::WRITING)
-		{
-			return (int64)DataArray.Num();
-		}
-		else
+		TSharedPtr<FSoundFileDataEntry> Entry = SoundFileManager->GetEntry(SoundFileHandle);
+		if (Entry.IsValid())
 		{
-			return DataBufferSize;
+			bOutIsStreamed = Entry->bIsStreamed;
+			return ESoundFileError::NONE;
 		}
+		return ESoundFileError::INVALID_DATA;
 	}
 
-	TArray<uint8>& FSoundFileData::GetDataArray()
-	{
-		// If we're accessing the array, make sure this sound file data object is in write-mode
-		check(Mode == ESoundFileOpenMode::WRITING);
-		return DataArray;
-	}
-
-	void FSoundFileData::SetData(uint8* InDataPtr, int64 InDataSize)
-	{
-		// Make sure this sound file data object wasn't opened for writing
-		check(Mode == ESoundFileOpenMode::READING);
-		DataBuffer = InDataPtr;
-		DataBufferSize = InDataSize;
-	}
-
-	// Exported Functions
-
-	TSharedPtr<ISoundFileData> CreateSoundFileData()
-	{
-		// Only internal sound file code creates write-able sound file data
-		return TSharedPtr<ISoundFileData>(new FSoundFileData(ESoundFileOpenMode::READING));
-	}
-
-	TSharedPtr<ISoundFile> CreateSoundFile()
-	{
-		return TSharedPtr<ISoundFile>(new FSoundFile());
-	}
-
-	// UnrealAudio Module Functions
-
-	static void* GetSoundFileDllHandle()
-	{
-		void* DllHandle = nullptr;
-#if PLATFORM_WINDOWS
-		FString Path = FPaths::EngineDir() / FString(TEXT("Binaries/ThirdParty/libsndfile/Win64/"));
-		FPlatformProcess::PushDllDirectory(*Path);
-		DllHandle = FPlatformProcess::GetDllHandle(*(Path + "libsndfile-1.dll"));
-		FPlatformProcess::PopDllDirectory(*Path);
-#elif PLATFORM_MAC
-//		FString Path = FPaths::EngineDir() / FString(TEXT("Binaries/ThirdParty/libsndfile/Mac/"));
-//		FPlatformProcess::PushDllDirectory(*Path);
-		DllHandle = FPlatformProcess::GetDllHandle(TEXT("libsndfile.1.dylib"));
-//		FPlatformProcess::PopDllDirectory(*Path);
-#endif
-		return DllHandle;
-	}
-
-	bool FUnrealAudioModule::LoadSoundFileLib()
+	void GetSoundFileListInDirectory(FString& Directory, TArray<FString>& SoundFiles, bool bRecursive /* = true */)
 	{
-		SoundFileDllHandle = GetSoundFileDllHandle();
-		if (!SoundFileDllHandle)
+		if (FPaths::DirectoryExists(Directory))
 		{
-			UE_LOG(LogUnrealAudio, Error, TEXT("Failed to load Sound File dll"));
-			return false;
-		}
-
-		bool bSuccess = true;
-
-		void* LambdaDLLHandle = SoundFileDllHandle;
-
-		// Helper function to load DLL exports and report warnings
-		auto GetSoundFileDllExport = [&LambdaDLLHandle](const TCHAR* FuncName, bool& bInSuccess) -> void*
-		{
-			if (bInSuccess)
+			// Get input files in formats we're currently supporting
+			IFileManager& FileManager = IFileManager::Get();
+			if (bRecursive)
 			{
-				void* FuncPtr = FPlatformProcess::GetDllExport(LambdaDLLHandle, FuncName);
-				if (FuncPtr == nullptr)
-				{
-					bInSuccess = false;
-					UE_LOG(LogUnrealAudio, Warning, TEXT("Failed to locate the expected DLL import function '%s' in the SoundFile DLL."), *FuncName);
-					FPlatformProcess::FreeDllHandle(LambdaDLLHandle);
-					LambdaDLLHandle = nullptr;
-				}
-				return FuncPtr;
+				FileManager.FindFilesRecursive(SoundFiles, *Directory, TEXT("*.wav"), true, false, false);
+				FileManager.FindFilesRecursive(SoundFiles, *Directory, TEXT("*.aif"), true, false, false);
+				FileManager.FindFilesRecursive(SoundFiles, *Directory, TEXT("*.flac"), true, false, false);
+				FileManager.FindFilesRecursive(SoundFiles, *Directory, TEXT("*.ogg"), true, false, false);
 			}
 			else
 			{
-				return nullptr;
+				FileManager.FindFiles(SoundFiles, *(Directory / "*.wav"), true, false);
+				FileManager.FindFiles(SoundFiles, *(Directory / "*.aif"), true, false);
+				FileManager.FindFiles(SoundFiles, *(Directory / "*.flac"), true, false);
+				FileManager.FindFiles(SoundFiles, *(Directory / "*.ogg"), true, false);
 			}
-		};
-
-		SoundFileOpen = (SoundFileOpenFuncPtr)GetSoundFileDllExport(TEXT("sf_open"), bSuccess);
-		SoundFileOpenVirtual = (SoundFileOpenVirtualFuncPtr)GetSoundFileDllExport(TEXT("sf_open_virtual"), bSuccess);
-		SoundFileClose = (SoundFileCloseFuncPtr)GetSoundFileDllExport(TEXT("sf_close"), bSuccess);
-		SoundFileError = (SoundFileErrorFuncPtr)GetSoundFileDllExport(TEXT("sf_error"), bSuccess);
-		SoundFileStrError = (SoundFileStrErrorFuncPtr)GetSoundFileDllExport(TEXT("sf_strerror"), bSuccess);
-		SoundFileErrorNumber = (SoundFileErrorNumberFuncPtr)GetSoundFileDllExport(TEXT("sf_error_number"), bSuccess);
-		SoundFileCommand = (SoundFileCommandFuncPtr)GetSoundFileDllExport(TEXT("sf_command"), bSuccess);
-		SoundFileFormatCheck = (SoundFileFormatCheckFuncPtr)GetSoundFileDllExport(TEXT("sf_format_check"), bSuccess);
-		SoundFileSeek = (SoundFileSeekFuncPtr)GetSoundFileDllExport(TEXT("sf_seek"), bSuccess);
-		SoundFileGetVersion = (SoundFileGetVersionFuncPtr)GetSoundFileDllExport(TEXT("sf_version_string"), bSuccess);
-		SoundFileReadFramesFloat = (SoundFileReadFramesFloatFuncPtr)GetSoundFileDllExport(TEXT("sf_readf_float"), bSuccess);
-		SoundFileReadFramesDouble = (SoundFileReadFramesDoubleFuncPtr)GetSoundFileDllExport(TEXT("sf_readf_double"), bSuccess);
-		SoundFileWriteFramesFloat = (SoundFileWriteFramesFloatFuncPtr)GetSoundFileDllExport(TEXT("sf_writef_float"), bSuccess);
-		SoundFileWriteFramesDouble = (SoundFileWriteFramesDoubleFuncPtr)GetSoundFileDllExport(TEXT("sf_writef_double"), bSuccess);
-		SoundFileReadSamplesFloat = (SoundFileReadSamplesFloatFuncPtr)GetSoundFileDllExport(TEXT("sf_read_float"), bSuccess);
-		SoundFileReadSamplesDouble = (SoundFileReadSamplesDoubleFuncPtr)GetSoundFileDllExport(TEXT("sf_read_double"), bSuccess);
-		SoundFileWriteSamplesFloat = (SoundFileWriteSamplesFloatFuncPtr)GetSoundFileDllExport(TEXT("sf_write_float"), bSuccess);
-		SoundFileWriteSamplesDouble = (SoundFileWriteSamplesDoubleFuncPtr)GetSoundFileDllExport(TEXT("sf_write_double"), bSuccess);
-
-		// make sure we're successful
-		check(bSuccess);
-		return bSuccess;
-	}
-
-	bool FUnrealAudioModule::ShutdownSoundFileLib()
-	{
-		if (SoundFileDllHandle)
-		{
-			FPlatformProcess::FreeDllHandle(SoundFileDllHandle);
-			SoundFileDllHandle = nullptr;
-			SoundFileOpen = nullptr;
-			SoundFileOpenVirtual = nullptr;
-			SoundFileClose = nullptr;
-			SoundFileError = nullptr;
-			SoundFileStrError = nullptr;
-			SoundFileErrorNumber = nullptr;
-			SoundFileCommand = nullptr;
-			SoundFileFormatCheck = nullptr;
-			SoundFileSeek = nullptr;
-			SoundFileGetVersion = nullptr;
-			SoundFileReadFramesFloat = nullptr;
-			SoundFileReadFramesDouble = nullptr;
-			SoundFileWriteFramesFloat = nullptr;
-			SoundFileWriteFramesDouble = nullptr;
-			SoundFileReadSamplesFloat = nullptr;
-			SoundFileReadSamplesDouble = nullptr;
-			SoundFileWriteSamplesFloat = nullptr;
-			SoundFileWriteSamplesDouble = nullptr;
 		}
-		return true;
 	}
 
 }
diff --git a/Engine/Source/Runtime/UnrealAudio/Private/UnrealAudioSoundFileConvert.cpp b/Engine/Source/Runtime/UnrealAudio/Private/UnrealAudioSoundFileConvert.cpp
new file mode 100644
--- /dev/null
+++ b/Engine/Source/Runtime/UnrealAudio/Private/UnrealAudioSoundFileConvert.cpp
@@ -0,0 +1,215 @@
+#include "UnrealAudioPrivate.h"
+#include "UnrealAudioSoundFile.h"
+#include "UnrealAudioModule.h"
+#include "ModuleManager.h"
+#include "UnrealAudioSoundFile.h"
+#include "UnrealAudioSoundFileInternal.h"
+#include "UnrealAudioSampleRateConverter.h"
+
+#if ENABLE_UNREAL_AUDIO
+
+
+#define SOUND_CONVERT_CHECK(Result)									\
+	if (Result != ESoundFileError::NONE)							\
+	{																\
+		OnError(Result);											\
+		return;														\
+	}
+
+
+namespace UAudio
+{
+	class FAsyncSoundFileConvertTask
+	{
+	public:
+		FAsyncSoundFileConvertTask(FUnrealAudioModule* InAudioModule, const FString* InSoundFilePath, const FString* OutSoundFilePath, const FSoundFileConvertFormat* InConvertFormat)
+			: AudioModule(InAudioModule)
+			, SoundFilePath(*InSoundFilePath)
+			, OutSoundFilePath(*OutSoundFilePath)
+			, ConvertFormat(*InConvertFormat)
+		{
+			AudioModule->IncrementBackgroundTaskCount();
+		}
+		
+		~FAsyncSoundFileConvertTask()
+		{
+			AudioModule->DecrementBackgroundTaskCount();
+		}
+
+		void DoWork()
+		{
+			// Synchronously load the input sound file
+			TSharedPtr<ISoundFile> InputSoundFile = AudioModule->LoadSoundFile(*SoundFilePath, false);
+			if (!InputSoundFile.IsValid())
+			{
+				return;
+			}
+
+			TSharedPtr<FSoundFileReader> InputSoundFileReader = AudioModule->CreateSoundFileReader();
+
+			ESoundFileError::Type Error = InputSoundFileReader->Init(InputSoundFile, false);
+			if (Error != ESoundFileError::NONE)
+			{
+				return;
+			}
+			check(Error == ESoundFileError::NONE);
+
+			// Cast to an internal object to get access to non-public API
+			FSoundFile* InputSoundFileInternal = static_cast<FSoundFile*>(InputSoundFile.Get());
+
+			FSoundFileDescription InputDescription;
+			Error = InputSoundFileInternal->GetDescription(InputDescription);
+			check(Error == ESoundFileError::NONE);
+
+			TArray<ESoundFileChannelMap::Type> ChannelMap;
+			Error = InputSoundFileInternal->GetChannelMap(ChannelMap);
+
+			FSoundFileDescription NewSoundFileDescription;
+			NewSoundFileDescription.NumChannels = InputDescription.NumChannels;
+			NewSoundFileDescription.NumFrames	= InputDescription.NumFrames;
+			NewSoundFileDescription.FormatFlags = ConvertFormat.Format;
+			NewSoundFileDescription.SampleRate	= ConvertFormat.SampleRate;
+			NewSoundFileDescription.NumSections = InputDescription.NumSections;
+			NewSoundFileDescription.bIsSeekable = InputDescription.bIsSeekable;
+
+			TSharedPtr<FSoundFileWriter> SoundFileWriter = AudioModule->CreateSoundFileWriter();
+			Error = SoundFileWriter->Init(NewSoundFileDescription, ChannelMap, ConvertFormat.EncodingQuality);
+			check(Error == ESoundFileError::NONE);
+
+			// Create a buffer to do the processing 
+			SoundFileCount ProcessBufferSamples = 1024 * NewSoundFileDescription.NumChannels;
+			TArray<float> ProcessBuffer;
+			ProcessBuffer.Init(0.0f, ProcessBufferSamples);
+
+			double SampleRateConversionRatio = (double)InputDescription.SampleRate / ConvertFormat.SampleRate;
+
+			FSampleRateConverter SampleRateConverter;
+			SampleRateConverter.Init(SampleRateConversionRatio, NewSoundFileDescription.NumChannels);
+
+			TArray<float> OutputBuffer;
+			SoundFileCount OutputBufferSamples = ProcessBufferSamples / SampleRateConversionRatio;
+			OutputBuffer.Reserve(OutputBufferSamples);
+
+			// Find the max value if we've been told to do peak normalization on import
+			float MaxValue = 0.0f;
+			SoundFileCount InputSamplesRead = 0;
+			bool bPerformPeakNormalization = ConvertFormat.bPerformPeakNormalization;
+			if (bPerformPeakNormalization)
+			{
+				Error = InputSoundFileReader->ReadSamples(ProcessBuffer.GetData(), ProcessBufferSamples, InputSamplesRead);
+				SOUND_CONVERT_CHECK(Error);
+
+				while (InputSamplesRead)
+				{
+					for (SoundFileCount Sample = 0; Sample < InputSamplesRead; ++Sample)
+					{
+						if (ProcessBuffer[Sample] > FMath::Abs(MaxValue))
+						{
+							MaxValue = ProcessBuffer[Sample];
+						}
+					}
+
+					Error = InputSoundFileReader->ReadSamples(ProcessBuffer.GetData(), ProcessBufferSamples, InputSamplesRead);
+					SOUND_CONVERT_CHECK(Error);
+				}
+
+				// If this happens, it means we have a totally silent file
+				if (MaxValue == 0.0)
+				{
+					bPerformPeakNormalization = false;
+				}
+
+				// Seek the file back to the beginning
+				SoundFileCount OutOffset;
+				InputSoundFileReader->SeekFrames(0, ESoundFileSeekMode::FROM_START, OutOffset);
+			}
+
+			bool SamplesProcessed = true;
+
+			// Read the first block of samples
+			Error = InputSoundFileReader->ReadSamples(ProcessBuffer.GetData(), ProcessBufferSamples, InputSamplesRead);
+			SOUND_CONVERT_CHECK(Error);
+
+			while (InputSamplesRead != 0)
+			{
+				SampleRateConverter.ProcessBlock(ProcessBuffer.GetData(), InputSamplesRead, OutputBuffer);
+
+				SoundFileCount SamplesWritten;
+				Error = SoundFileWriter->WriteSamples((const float*)OutputBuffer.GetData(), OutputBuffer.Num(), SamplesWritten);
+				SOUND_CONVERT_CHECK(Error);
+				DEBUG_AUDIO_CHECK(SamplesWritten == OutputBuffer.Num());
+
+				OutputBuffer.Reset();
+
+				// read more samples
+				Error = InputSoundFileReader->ReadSamples(ProcessBuffer.GetData(), ProcessBufferSamples, InputSamplesRead);
+				SOUND_CONVERT_CHECK(Error);
+
+				// ... normalize the samples if we're told to
+				if (bPerformPeakNormalization)
+				{
+					for (int32 Sample = 0; Sample < InputSamplesRead; ++Sample)
+					{
+						ProcessBuffer[Sample] /= MaxValue;
+					}
+				}
+			}
+
+			// Release the sound file handles as soon as we finished converting the file
+			InputSoundFileReader->Release();
+			SoundFileWriter->Release();
+
+			// Get the raw binary data
+			TArray<uint8>* Data = nullptr;
+			SoundFileWriter->GetData(&Data);
+
+			// Write the data to the output file
+			bool bSuccess = FFileHelper::SaveArrayToFile(*Data, *OutSoundFilePath);
+			check(bSuccess);
+		}
+
+
+		bool CanAbandon()
+		{ 
+			return true; 
+		}
+		
+		void Abandon() 
+		{
+		}
+
+		void OnError(ESoundFileError::Type Error)
+		{
+//			((ISoundFileInternal*)SoundFile.Get())->SetError(Error);
+		}
+
+		FORCEINLINE TStatId GetStatId() const
+		{
+			RETURN_QUICK_DECLARE_CYCLE_STAT(FAsyncSoundFileConvertTask, STATGROUP_ThreadPoolAsyncTasks);
+		}
+
+		FUnrealAudioModule* AudioModule;
+		FString SoundFilePath;
+		FString OutSoundFilePath;
+		FSoundFileConvertFormat ConvertFormat;
+		friend class FAsyncTask<FAsyncSoundFileConvertTask>;
+	};
+
+	void FUnrealAudioModule::ConvertSound(const FString& InputFilePath, const FString& OutputFilePath, const FSoundFileConvertFormat& ConvertFormat)
+	{
+		// Make an auto-deleting task to perform the sound file loading work in the background task thread pool
+		FAutoDeleteAsyncTask<FAsyncSoundFileConvertTask>* Task = new FAutoDeleteAsyncTask<FAsyncSoundFileConvertTask>(this, &InputFilePath, &OutputFilePath, &ConvertFormat);
+		Task->StartBackgroundTask();
+	}
+
+	void FUnrealAudioModule::ConvertSound(const FString& InputFilePath, const FString& OutputFilePath)
+	{
+		// Make an auto-deleting task to perform the sound file loading work in the background task thread pool
+		FAutoDeleteAsyncTask<FAsyncSoundFileConvertTask>* Task = new FAutoDeleteAsyncTask<FAsyncSoundFileConvertTask>(this, &InputFilePath, &OutputFilePath, &DefaultConvertFormat);
+		Task->StartBackgroundTask();
+	}
+
+}
+
+#endif // #if ENABLE_UNREAL_AUDIO
+
diff --git a/Engine/Source/Runtime/UnrealAudio/Private/UnrealAudioSoundFileInternal.h b/Engine/Source/Runtime/UnrealAudio/Private/UnrealAudioSoundFileInternal.h
--- a/Engine/Source/Runtime/UnrealAudio/Private/UnrealAudioSoundFileInternal.h
+++ b/Engine/Source/Runtime/UnrealAudio/Private/UnrealAudioSoundFileInternal.h
@@ -3,323 +3,153 @@
 #pragma once 
 
 #include "UnrealAudioModule.h"
+#include "UnrealAudioPrivate.h"
+#include "UnrealAudioSoundFileManager.h"
 #include "UnrealAudioTypes.h"
 #include "UnrealAudioSoundFile.h"
+#include "UnrealAudioEntityManager.h"
+#include "BulkData.h"
 
 #if ENABLE_UNREAL_AUDIO
 
 namespace UAudio
 {
+
 	typedef int64 SoundFileCount;
-	typedef struct SoundFileHandleOpaque SoundFileHandle;
+	typedef struct SoundFileHandleOpaque LibSoundFileHandle;
 
-	/*
-	* ESoundFileSeekMode
-	* Specifies the how to seek a sound file
-	**/
 	namespace ESoundFileSeekMode
 	{
 		enum Type
 		{
-			// Seek the file from the start of the file.
 			FROM_START = 0,
-
-			// Seek the file from the current read/write position.
 			FROM_CURRENT = 1,
-
-			// Seek the file from the end position.
 			FROM_END = 2,
 		};
 	}
 
-	/*
-	* ESoundFileOpenMode
-	* Specifies the how to open a sound file1
-	**/
 	namespace ESoundFileOpenMode
 	{
 		enum Type
 		{
-			// The sound file is open for reading.
 			READING = 0x10,
-
-			// The sound file is open for writing.
 			WRITING = 0x20,
-
 			UNKNOWN = 0,
 		};
 	}
 
-
-	/**
-	ISoundFileVirtual
-	Interface for specifying the virtual sound file i/o operations.
-	*/
-	class ISoundFileVirtual : public ISoundFile
+	class ISoundFileParser
 	{
 	public:
-		virtual ~ISoundFileVirtual() {}
+		virtual ~ISoundFileParser() {}
 
-		/**
-		* Gets the length of the file in bytes.
-		* @param OutLength The returned length of the file.
-		*/
 		virtual ESoundFileError::Type GetLengthBytes(SoundFileCount& OutLength) const = 0;
-
-		/**
-		* Seeks the read/write position offset in bytes.
-		* @param Offset The offset in bytes of the file
-		* @param SeekMode The type of seek to use when performing the seek.
-		* @param OutOffset Optionally retrieve the actual offset in bytes after the seek operation.
-		*/
 		virtual ESoundFileError::Type SeekBytes(SoundFileCount Offset, ESoundFileSeekMode::Type SeekMode, SoundFileCount& OutOffset) = 0;
-
-		/**
-		* Reads the number of bytes from the given buffer ptr.
-		* @param DataPtr Pointer to a buffer for copying into.
-		* @param NumBytes The number of bytes to read.
-		* @param NumBytesRead Returns the number of bytes actually read.
-		*/
 		virtual ESoundFileError::Type ReadBytes(void* DataPtr, SoundFileCount NumBytes, SoundFileCount& NumBytesRead) = 0;
-
-		/**
-		* Writes the number of bytes into the given buffer ptr.
-		* @param DataPtr Pointer to an input buffer
-		* @param NumBytes The number of bytes to write.
-		* @param NumBytesWritten The number of bytes actually written.
-		*/
 		virtual ESoundFileError::Type WriteBytes(const void* DataPtr, SoundFileCount NumBytes, SoundFileCount& NumBytesWritten) = 0;
-
-		/** @return Gets the current file offset in bytes. */
-		virtual SoundFileCount GetOffsetBytes() const = 0;
+		virtual ESoundFileError::Type GetOffsetBytes(SoundFileCount& OutOffset) const = 0;
 	};
 
-	/**
-	FSoundFile
-	Internal SoundFile API
-	*/
-	class ISoundFileInternal : public ISoundFileVirtual
+	class FSoundFileReader : public ISoundFileParser
 	{
 	public:
-		virtual ~ISoundFileInternal() {}
-
-		/**
-		* Sets the error on the sound file. Internal audio systems can cause audio files to get errors. For example, if a failure occurs
-		* during import due to an error on the input sound file, the output (imported) sound file needs to also get the error.
-		*/
-		virtual ESoundFileError::Type SetError(ESoundFileError::Type Error) = 0;
-
-		/**
-		* Opens a sound file from the given path.
-		* @param FilePath Path to the sound file.
-		* @param Mode The mode to open the sound file in.
-		*/
-		virtual ESoundFileError::Type OpenFileForReading(const FString& FilePath) = 0;
-
-		/**
-		* Opens an empty virtual sound file of the given description for writing to.
-		* @param FileDescription Description of the type of file this is to be.
-		* @param InChannelMap The input channel map to use for the file.
-		* @param Mode The mode to open the sound file in.
-		*/
-		virtual ESoundFileError::Type OpenEmptyFileForImport(const FSoundFileDescription& FileDescription, const TArray<ESoundFileChannelMap::Type>& InChannelMap) = 0;
+		FSoundFileReader(FUnrealAudioModule* AudioModule);
+		~FSoundFileReader();
 
-		/**
-		* Opens an empty sound file on disk of the given description for writing to.
-		* @param FileDescription Description of the type of file this is to be.
-		* @param InChannelMap The input channel map to use for the file.
-		* @param Mode The mode to open the sound file in.
-		*/
-		virtual ESoundFileError::Type OpenEmptyFileForExport(const FString& FilePath, const FSoundFileDescription& FileDescription, const TArray<ESoundFileChannelMap::Type>& InChannelMap) = 0;
-
-		/** Releases the internal sound file handle if its open */
-		virtual ESoundFileError::Type ReleaseSoundFileHandle() = 0;
-
-		/**
-		* Sets the description of the original file used to import this file.
-		* @param FileDescription Description of the original file used to import this file.
-		*/
-		virtual ESoundFileError::Type SetImportFileInfo(const FSoundFileDescription& FileDescription, const FString& FullPath) = 0;
-
-		/**
-		* Sets the encoding quality of a sound file.
-		* @param EncodingQuality The encoding quality of a sound file (between 0.0 (low) and 1.0 (high))
-		* @note This can only be set during sound file import before importing begins, but after the file has been opened.
-		*		It can also only be set on OGG-Vorbis files.
-		*/
-		virtual ESoundFileError::Type SetEncodingQuality(double EncodingQuality) = 0;
-
-		/**
-		* Sets the state of the sound file to begin loading
-		*/
-		virtual ESoundFileError::Type BeginImport() = 0;
-
-		/**
-		* Sets the state of the sound file to end loading
-		*/
-		virtual ESoundFileError::Type EndImport() = 0;
-
-		/**
-		* Seeks the read/write position offset in frames.
-		* @param Offset The offset in frames of the file
-		* @param SeekMode The type of seek to use when performing the seek.
-		* @param OutOffset Optionally retrieve the actual offset in bytes after the seek operation.
-		*/
-		virtual ESoundFileError::Type SeekFrames(SoundFileCount Offset, ESoundFileSeekMode::Type SeekMode, SoundFileCount& OutOffset) = 0;
-
-		/**
-		* Reads out the number of frames of data as interleaved floats/doubles.
-		* @param DataPtr Pointer to a buffer to copy data into.
-		* @param NumFrames The number of frames to read.
-		* @param NumFramesRead Returns the number of frames actually read.
-		*/
-		virtual ESoundFileError::Type ReadFrames(float* DataPtr, SoundFileCount NumFrames, SoundFileCount& NumFramesRead) = 0;
-		virtual ESoundFileError::Type ReadFrames(double* DataPtr, SoundFileCount NumFrames, SoundFileCount& NumFramesRead) = 0;
-
-		/**
-		* Reads out the number of samples of data as floats/doubles.
-		* @param DataPtr Pointer to a buffer for copy data into.
-		* @param NumSamples The number of samples to read.
-		* @param OutNumSamplesRead Returns the number of samples actually read.
-		*/
-		virtual ESoundFileError::Type ReadSamples(float* DataPtr, SoundFileCount NumSamples, SoundFileCount& OutNumSamplesRead) = 0;
-		virtual ESoundFileError::Type ReadSamples(double* DataPtr, SoundFileCount NumSamples, SoundFileCount& OutNumSamplesRead) = 0;
+		// ISoundFileParser
+		ESoundFileError::Type GetLengthBytes(SoundFileCount& OutLength) const override;
+		ESoundFileError::Type SeekBytes(SoundFileCount Offset, ESoundFileSeekMode::Type SeekMode, SoundFileCount& OutOffset) override;
+		ESoundFileError::Type ReadBytes(void* DataPtr, SoundFileCount NumBytes, SoundFileCount& OutNumBytesRead) override;
+		ESoundFileError::Type WriteBytes(const void* DataPtr, SoundFileCount NumBytes, SoundFileCount& OutNumBytesWritten) override;
+		ESoundFileError::Type GetOffsetBytes(SoundFileCount& OutOffset) const override;
 
-		/**
-		* Writes the number of frames of data as interleaved floats/doubles.
-		* @param DataPtr Pointer to a buffer for copying data from.
-		* @param NumFrames The number of frames to write.
-		* @param OutNumFramesWritten Returns the number of frames actually written.
-		*/
-		virtual ESoundFileError::Type WriteFrames(const float* Data, SoundFileCount NumFrames, SoundFileCount& OutNumFramesWritten) = 0;
-		virtual ESoundFileError::Type WriteFrames(const double* Data, SoundFileCount NumFrames, SoundFileCount& OutNumFramesWritten) = 0;
+		virtual ESoundFileError::Type Init(TSharedPtr<ISoundFile> InSoundFileData, bool bIsStreamed);
+		virtual ESoundFileError::Type Release();
+		virtual ESoundFileError::Type SeekFrames(SoundFileCount Offset, ESoundFileSeekMode::Type SeekMode, SoundFileCount& OutOffset);
+		virtual ESoundFileError::Type ReadFrames(float* DataPtr, SoundFileCount NumFrames, SoundFileCount& OutNumFramesRead);
+		virtual ESoundFileError::Type ReadFrames(double* DataPtr, SoundFileCount NumFrames, SoundFileCount& OutNumFramesRead);
+		virtual ESoundFileError::Type ReadSamples(float* DataPtr, SoundFileCount NumSamples, SoundFileCount& OutNumSamplesRead);
+		virtual ESoundFileError::Type ReadSamples(double* DataPtr, SoundFileCount NumSamples, SoundFileCount& OutNumSamplesRead);
 
-		/**
-		* Writes the number of samples of data as floats/doubles.
-		* @param DataPtr Pointer to a buffer for copying data from.
-		* @param NumSamples The number of samples to write.
-		* @param OutNumSamplesWritten Returns the number of samples actually written.
-		*/
-		virtual ESoundFileError::Type WriteSamples(const float* Data, SoundFileCount NumSamples, SoundFileCount& OutNumSamplesWritten) = 0;
-		virtual ESoundFileError::Type WriteSamples(const double* Data, SoundFileCount NumSamples, SoundFileCount& OutNumSamplesWritten) = 0;
+	private:
+		ESoundFileError::Type InitLoaded(TSharedPtr<ISoundFile> InSoundFileData);
+		ESoundFileError::Type InitStreamed(TSharedPtr<ISoundFile> InSoundFileData);
+		ESoundFileError::Type SetError(ESoundFileError::Type InError);
+
+		TSharedPtr<ISoundFile>	SoundFileData;
+		FUnrealAudioModule*		AudioModule;
+		SoundFileCount			CurrentIndexBytes;
+		LibSoundFileHandle*		FileHandle;
+		FThreadSafeCounter		State;
+		FThreadSafeCounter		CurrentError;
 	};
 
-
-	/**
-	FSoundFile
-	Concrete implementation of ISoundFile which uses libsndfile for file reading/writing.
-	*/
-	class FSoundFile : public ISoundFileInternal
+	class FSoundFileWriter : public ISoundFileParser
 	{
 	public:
-		FSoundFile();
-		~FSoundFile();
-
-		// ISoundFile implementation
-		ESoundFileError::Type Initialize(TSharedPtr<ISoundFileData> InSoundFileData) override;
-		ESoundFileError::Type GetName(FString& Name) override;
-		ESoundFileError::Type GetSoundFileData(TSharedPtr<ISoundFileData>& OutSoundFileData) override;
-		ESoundFileState::Type GetState() const override;
-		ESoundFileError::Type GetError() const override;
-
-		// ISoundFileInternal implementation
-		ESoundFileError::Type SetError(ESoundFileError::Type Error) override;
-		ESoundFileError::Type OpenFileForReading(const FString& FilePath) override;
-		ESoundFileError::Type OpenEmptyFileForImport(const FSoundFileDescription& FileDescription, const TArray<ESoundFileChannelMap::Type>& InChannelMap) override;
-		ESoundFileError::Type OpenEmptyFileForExport(const FString& FilePath, const FSoundFileDescription& FileDescription, const TArray<ESoundFileChannelMap::Type>& InChannelMap) override;
-		ESoundFileError::Type ReleaseSoundFileHandle() override;
+		FSoundFileWriter(FUnrealAudioModule* AudioModule);
+		~FSoundFileWriter();
 
-		ESoundFileError::Type SetImportFileInfo(const FSoundFileDescription& FileDescription, const FString& FullPath) override;
-		ESoundFileError::Type SetEncodingQuality(double EncodingQuality) override;
-		ESoundFileError::Type BeginImport() override;
-		ESoundFileError::Type EndImport() override;
-		ESoundFileError::Type SeekFrames(SoundFileCount Offset, ESoundFileSeekMode::Type SeekMode, SoundFileCount& OutOffset) override;
-		ESoundFileError::Type ReadFrames(float* DataPtr, SoundFileCount NumFrames, SoundFileCount& OutNumFramesRead) override;
-		ESoundFileError::Type ReadFrames(double* DataPtr, SoundFileCount NumFrames, SoundFileCount& OutNumFramesRead) override;
-		ESoundFileError::Type ReadSamples(float* DataPtr, SoundFileCount NumSamples, SoundFileCount& OutNumSamplesRead) override;
-		ESoundFileError::Type ReadSamples(double* DataPtr, SoundFileCount NumSamples, SoundFileCount& OutNumSamplesRead) override;
-		ESoundFileError::Type WriteFrames(const float* Data, SoundFileCount NumFrames, SoundFileCount& OutNumFramesWritten) override;
-		ESoundFileError::Type WriteFrames(const double* Data, SoundFileCount NumFrames, SoundFileCount& OutNumFramesWritten) override;
-		ESoundFileError::Type WriteSamples(const float* DataPtr, SoundFileCount NumSamples, SoundFileCount& OutNumSampleWritten) override;
-		ESoundFileError::Type WriteSamples(const double* DataPtr, SoundFileCount NumSamples, SoundFileCount& OutNumSampleWritten) override;
-		
-		// ISoundFileVirtual implementation
+		// ISoundFileParser
 		ESoundFileError::Type GetLengthBytes(SoundFileCount& OutLength) const override;
 		ESoundFileError::Type SeekBytes(SoundFileCount Offset, ESoundFileSeekMode::Type SeekMode, SoundFileCount& OutOffset) override;
 		ESoundFileError::Type ReadBytes(void* DataPtr, SoundFileCount NumBytes, SoundFileCount& OutNumBytesRead) override;
 		ESoundFileError::Type WriteBytes(const void* DataPtr, SoundFileCount NumBytes, SoundFileCount& OutNumBytesWritten) override;
-		SoundFileCount GetOffsetBytes() const override;
-
-	private:
-
-		TSharedPtr<ISoundFileData> SoundFileData;
-
-		/** The current read/write index in the raw data */
-		SoundFileCount CurrentIndexBytes;
-
-		/** Handle to the sound file decoder. */
-		SoundFileHandle* FileHandle;
-
-		/** The current state of the sound file. */
-		volatile ESoundFileState::Type State;
+		ESoundFileError::Type GetOffsetBytes(SoundFileCount& OutOffset) const override;
 
-		/** The last error of the sound file. */
-		volatile ESoundFileError::Type CurrentError;
+		virtual ESoundFileError::Type Init(const FSoundFileDescription& FileDescription, const TArray<ESoundFileChannelMap::Type>& InChannelMap, double EncodingQuality);
+		virtual ESoundFileError::Type Release();
+		virtual ESoundFileError::Type SeekFrames(SoundFileCount Offset, ESoundFileSeekMode::Type SeekMode, SoundFileCount& OutOffset);
+		virtual ESoundFileError::Type WriteFrames(const float* Data, SoundFileCount NumFrames, SoundFileCount& OutNumFramesWritten);
+		virtual ESoundFileError::Type WriteFrames(const double* Data, SoundFileCount NumFrames, SoundFileCount& OutNumFramesWritten);
+		virtual ESoundFileError::Type WriteSamples(const float* DataPtr, SoundFileCount NumSamples, SoundFileCount& OutNumSampleWritten);
+		virtual ESoundFileError::Type WriteSamples(const double* DataPtr, SoundFileCount NumSamples, SoundFileCount& OutNumSampleWritten);
+		virtual ESoundFileError::Type GetData(TArray<uint8>** OutData);
 
-		int32 Id;
+	private:
+		ESoundFileError::Type SetError(ESoundFileError::Type InError);
+
+		FUnrealAudioModule*					AudioModule;
+		SoundFileCount						CurrentIndexBytes;
+		LibSoundFileHandle*					FileHandle;
+		FSoundFileDescription				Description;
+		TArray<ESoundFileChannelMap::Type>	ChannelMap;
+		TArray<uint8>						BulkData;
+		double								EncodingQuality;
+		FThreadSafeCounter					State;
+		FThreadSafeCounter					CurrentError;
 	};
 
-	/**
-	FSoundFileData
-	*/
-	class FSoundFileData : public ISoundFileData
+
+	class FSoundFile : public ISoundFile
 	{
 	public:
-		FSoundFileData(ESoundFileOpenMode::Type Mode);
-		~FSoundFileData();
-
-		// ISoundFileData interface
-		const FString& GetFilePath() const override;
-		void SetFilePath(const FString& InPath) override;
-		const FSoundFileDescription& GetDescription() const override;
-		void SetDescription(const FSoundFileDescription& InDescription) override;
-		const FSoundFileDescription& GetOriginalDescription() const override;
-		void SetOriginalDescription(const FSoundFileDescription& InDescription) override;
-		const TArray<ESoundFileChannelMap::Type>& GetChannelMap() const override;
-		void SetChannelMap(const TArray<ESoundFileChannelMap::Type>& InChannelMap) override;
-		const uint8* GetData() const override;
-		int64 GetDataSize() const override;
-		void SetData(uint8* InDataPtr, int64 InDataSize) override;
+		FSoundFile(class FSoundFileManager* InSoundFileManager);
+		~FSoundFile();
 
-		TArray<uint8>& GetDataArray();
-		void SetNumFrames(int64 InNumFrames);
+		ESoundFileError::Type GetState(ESoundFileState::Type& OutState) const override;
+		ESoundFileError::Type GetError() const override;
+		ESoundFileError::Type GetId(uint32& OutId) const override;
+		ESoundFileError::Type GetPath(FName& OutPath) const override;
+		ESoundFileError::Type GetBulkData(TArray<uint8>** OutData) const override;
+		ESoundFileError::Type GetDataSize(int32& DataSize) const override;
+		ESoundFileError::Type GetDescription(FSoundFileDescription& OutDescription) const override;
+		ESoundFileError::Type GetChannelMap(TArray<ESoundFileChannelMap::Type>& OutChannelMap) const override;
+		ESoundFileError::Type IsStreamed(bool& bOutIsStreamed) const override;
+
+		// Internal API
+		ESoundFileError::Type Init(const FSoundFileHandle& InSoundFileHandle);
+		ESoundFileError::Type SetError(ESoundFileError::Type InError);
 
 	private:
-		// The path of the file when it was imported
-		FString FilePath;
-
-		// The current description of the file (i.e. compressed)
-		FSoundFileDescription Description;
-
-		// The description of the file before it was imported
-		FSoundFileDescription OriginalDescription;
-
-		// The channel map of the file
-		TArray<ESoundFileChannelMap::Type> ChannelMap;
-
-		// An array of data, used to build the file data during import
-		TArray<uint8> DataArray;
-
-		// A ptr to a buffer of data, used to point this sound file data to an arbitrary location in memory
-		// This sound file data object does not own this data buffer
-		uint8* DataBuffer;
+		FSoundFileHandle SoundFileHandle;
+		class FSoundFileManager* SoundFileManager;
+		ESoundFileError::Type Error;
+	};
 
-		// The size of the data buffer in bytes
-		int64 DataBufferSize;
+	ESoundFileError::Type LoadSoundFileFromPath(const FString& FilePath, FSoundFileDescription& Description, TArray<ESoundFileChannelMap::Type>& ChannelMap, TArray<uint8>& BulkData);
+	ESoundFileError::Type GetSoundFileInfoFromPath(const FString& FilePath, FSoundFileDescription& Description, TArray<ESoundFileChannelMap::Type>& ChannelMap);
 
-		// The mode this sound file data is in (read or write)
-		ESoundFileOpenMode::Type Mode;
-	};
 }
 
 
diff --git a/Engine/Source/Runtime/UnrealAudio/Private/UnrealAudioSoundFileManager.cpp b/Engine/Source/Runtime/UnrealAudio/Private/UnrealAudioSoundFileManager.cpp
new file mode 100644
--- /dev/null
+++ b/Engine/Source/Runtime/UnrealAudio/Private/UnrealAudioSoundFileManager.cpp
@@ -0,0 +1,548 @@
+// Copyright 1998-2015 Epic Games, Inc. All Rights Reserved.
+
+#pragma once 
+
+#include "UnrealAudioPrivate.h"
+#include "UnrealAudioModule.h"
+#include "UnrealAudioTypes.h"
+#include "UnrealAudioSoundFile.h"
+#include "UnrealAudioSoundFileManager.h"
+
+#if ENABLE_UNREAL_AUDIO
+
+namespace UAudio
+{
+
+	FSoundFileDataEntry::FSoundFileDataEntry(class FSoundFileManager* InSoundFileManager, const FSoundFileHandle& InSoundFileHandle, bool bInIsStreamed)
+		: SoundFileManager(InSoundFileManager)
+		, SoundFileHandle(InSoundFileHandle)
+		, SoundFileState((int32)(ESoundFileState::UNINITIALIZED))
+		, Error(ESoundFileError::NONE)
+		, TimeSinceUsed(0.0f)
+		, NumReferences(1)
+		, bIsStreamed(bInIsStreamed)
+	{}
+
+	FSoundFileDataEntry::~FSoundFileDataEntry()
+	{
+		// Make sure nobody is referencing this entry
+		check(NumReferences == 0);
+
+		// Release this sound entry's file handle when the entry is deleted.
+		SoundFileManager->EntityManager.ReleaseEntity(SoundFileHandle);
+	}
+
+	FAsyncSoundFileLoadTask::FAsyncSoundFileLoadTask(FUnrealAudioModule* InAudioModule, FSoundFileHandle InSoundFileHandle, const FString& InPath)
+		: AudioModule(InAudioModule)
+		, SoundFileHandle(InSoundFileHandle)
+		, Path(InPath)
+	{
+		// make sure the audio module knows we're doing background task work so it doesn't shutdown while we're working here
+		AudioModule->IncrementBackgroundTaskCount();
+	}
+
+	FAsyncSoundFileLoadTask::~FAsyncSoundFileLoadTask()
+	{
+		// Let the audio module know we're done doing this background task
+		AudioModule->DecrementBackgroundTaskCount();
+	}
+
+	void FAsyncSoundFileLoadTask::DoWork()
+	{
+		FSoundFileManager& SoundFileManager = AudioModule->GetSoundFileManager();
+		TSharedPtr<FSoundFileDataEntry> DataEntry = SoundFileManager.SoundFileData[SoundFileHandle.GetIndex()];
+		check(DataEntry->SoundFileHandle.Id == SoundFileHandle.Id);
+
+		SoundFileManager.LoadSoundFileDataEntry(DataEntry);
+	}
+
+	/************************************************************************/
+	/* FSoundFileManager Implementation										*/
+	/************************************************************************/
+
+	FSoundFileManager::FSoundFileManager(FUnrealAudioModule* InAudioModule)
+		: EntityManager(500)
+		, AudioModule(InAudioModule)
+		, FileLoadingThreadPool(nullptr)
+		, NumSoundFilesLoaded(0)
+		, NumSoundFilesStreamed(0)
+		, NumBytesLoaded(0)
+		, NumActiveSounds(0)
+		, NumInactiveSounds(0)
+		, bLogOverMemoryTarget(true)
+	{
+		Settings = { 0 };
+	}
+
+	FSoundFileManager::~FSoundFileManager()
+	{
+		check(FileLoadingThreadPool == nullptr);
+	}
+
+	void FSoundFileManager::Init(FSoundFileManagerSettings& InSettings)
+	{
+		check(AudioModule != nullptr);
+
+		Settings = InSettings;
+
+		// Initialize sound file manager data
+		SoundFileHandles.Init(FSoundFileHandle(), InSettings.MaxNumberOfLoadedSounds);
+		SoundFileData.Init(nullptr, InSettings.MaxNumberOfLoadedSounds);
+
+		// Initialize the sound file loading pool
+		FileLoadingThreadPool = FQueuedThreadPool::Allocate();
+		FileLoadingThreadPool->Create(Settings.NumLoadingThreads, (32 * 1024), Settings.LoadingThreadPriority);
+	}
+
+	void FSoundFileManager::Shutdown()
+	{
+		// Clean up the thread pool
+		check(FileLoadingThreadPool);
+		FileLoadingThreadPool->Destroy();
+		FileLoadingThreadPool = nullptr;
+	}
+
+	TSharedPtr<ISoundFile> FSoundFileManager::CreateSoundFile()
+	{
+		return TSharedPtr<ISoundFile>(new FSoundFile(this));
+	}
+
+	TSharedPtr<ISoundFile> FSoundFileManager::GetSoundFile(const struct FSoundFileHandle& SoundFileHandle)
+	{
+		uint32 DataIndex = SoundFileHandle.GetIndex();
+		check(DataIndex < (uint32)SoundFileData.Num());
+		TSharedPtr<FSoundFileDataEntry> Entry = SoundFileData[DataIndex];
+
+		// Make sure the data at this index is valid
+		check(Entry.IsValid());
+		check(Entry->SoundFileHandle.Id == SoundFileHandle.Id);
+
+		// Bump up the ref-count of this data (so we know how many ISoundFile's are currently using this data)
+		Entry->NumReferences++;
+
+		// Make sure we reset the time-since-used parameter since we're now referencing this asset
+		Entry->TimeSinceUsed = 0.0;
+
+		TSharedPtr<ISoundFile> NewSoundFile = TSharedPtr<ISoundFile>(new FSoundFile(this));
+		(static_cast<FSoundFile*>(NewSoundFile.Get()))->Init(SoundFileHandle);
+
+		return NewSoundFile;
+	}
+
+	TSharedPtr<ISoundFile> FSoundFileManager::LoadSoundFile(const FName& Name, TArray<uint8>& InBulkData)
+	{
+		// TODO
+		return nullptr;
+	}
+
+	TSharedPtr<ISoundFile> FSoundFileManager::LoadSoundFile(const FName& Path, bool bLoadAsync)
+	{
+		// First see if we've already loaded this sound by checking the sound file name
+		FSoundFileHandle* ExistingHandle = NameToLoadedSoundMap.Find(Path);
+		if (ExistingHandle)
+		{
+			return GetSoundFile(*ExistingHandle);
+		}
+
+		return CreateNewSoundFile(Path, false, bLoadAsync);
+	}
+
+	TSharedPtr<ISoundFile> FSoundFileManager::StreamSoundFile(const FName& Path, bool bLoadAsync)
+	{
+		FSoundFileHandle* ExistingHandle = NameToStreamedSoundMap.Find(Path);
+		if (ExistingHandle)
+		{
+			return GetSoundFile(*ExistingHandle);
+		}
+
+		return CreateNewSoundFile(Path, true, bLoadAsync);
+	}
+
+	TSharedPtr<ISoundFile> FSoundFileManager::CreateNewSoundFile(const FName& Path, bool bIsStreamed, bool bLoadAsync)
+	{
+		// If we haven't already loaded this sound file (or if it's been unloaded)
+		// we need to create a data entry and trigger a loading operation.
+
+		// Create a new sound file handle
+		FSoundFileHandle NewHandle(EntityManager.CreateEntity());
+
+		// Confirm that the data at the index value in the new handle doesn't already have data
+		uint32 NewDataIndex = NewHandle.GetIndex();
+		check(!SoundFileData[NewDataIndex].IsValid());
+
+		// Create a new entry
+		TSharedPtr<FSoundFileDataEntry> NewEntry = TSharedPtr<FSoundFileDataEntry>(new FSoundFileDataEntry(this, NewHandle, bIsStreamed));
+
+		NewEntry->SoundFilePath = Path;
+
+		// Store it in the data array
+		SoundFileData[NewDataIndex] = NewEntry;
+
+		// Add the entry to the sound file name to handle map
+		if (bIsStreamed)
+		{
+			++NumSoundFilesStreamed;
+			NameToStreamedSoundMap.Add(Path, NewHandle);
+		}
+		else
+		{
+			++NumSoundFilesLoaded;
+			NameToLoadedSoundMap.Add(Path, NewHandle);
+		}
+
+		// Set the sound file state to loading...
+		NewEntry->SoundFileState.Set((int32)ESoundFileState::LOADING);
+
+		// Reset the last used for this entry
+		NewEntry->TimeSinceUsed = 0.0f;
+
+		if (bLoadAsync)
+		{
+			// Create a task to load the sound file. Note this task will delete itself when it finishes loading the sound file
+			FAsyncTask<FAsyncSoundFileLoadTask>* Task = new FAsyncTask<FAsyncSoundFileLoadTask>(AudioModule, NewHandle, Path.GetPlainNameString());
+			Task->StartBackgroundTask(FileLoadingThreadPool);
+		}
+		else
+		{
+			// Immediately load the data entry if this is a synchronous load call
+			LoadSoundFileDataEntry(NewEntry);
+		}
+
+		// Create a new sound file ptr and give it the handle to the internal data entry
+		TSharedPtr<ISoundFile> NewSoundFile = TSharedPtr<ISoundFile>(new FSoundFile(this));
+		(static_cast<FSoundFile*>(NewSoundFile.Get()))->Init(NewHandle);
+
+		return NewSoundFile;
+	}
+
+	ESoundFileState::Type FSoundFileManager::GetState(const FSoundFileHandle& SoundFileHandle) const
+	{
+		if (EntityManager.IsValidEntity(SoundFileHandle))
+		{
+			uint32 DataIndex = SoundFileHandle.GetIndex();
+
+			// If there is valid data at the data index
+			if (SoundFileData[DataIndex].IsValid())
+			{
+				return (ESoundFileState::Type)(SoundFileData[DataIndex]->SoundFileState.GetValue());
+			}
+		}
+		return ESoundFileState::UNINITIALIZED;
+	}
+
+	void FSoundFileManager::Update()
+	{
+		MainThreadChecker.CheckThread();
+
+		// Loop through sound files and update their TimeSinceUsed variable 
+
+		uint32 SoundFileEntryCount = 0;
+		uint32 NumEntries = NumSoundFilesStreamed + NumSoundFilesLoaded;
+		LeastRecentlyUsedSoundFiles.Reset();
+		SoundFileEntryComparePredicate Predicate;
+		bool bAddToLRUList = true;
+
+		// Reset the number of bytes loaded
+		NumBytesLoaded = 0;
+		NumActiveSounds = 0;
+		NumInactiveSounds = 0;
+
+		for (int32 i = 0; i < SoundFileData.Num() && SoundFileEntryCount < NumEntries; ++i)
+		{
+			if (!SoundFileData[i].IsValid())
+			{
+				continue;
+			}
+
+			bAddToLRUList = true;
+			SoundFileEntryCount++;
+			TSharedPtr<FSoundFileDataEntry> Entry = SoundFileData[i];
+
+			// If the entry is no longer referenced...
+			if (Entry->NumReferences == 0 && (Entry->SoundFileState.GetValue() != ESoundFileState::LOADING))
+			{
+				// If this entry no longer has references and it was streamed (i.e. no bulk data loaded)
+				// Then just remove this entry from our array altogether
+				if (Entry->BulkData.Num() == 0)
+				{
+					bAddToLRUList = false;
+					FlushSoundFileDataIndex(i);
+				}
+				else
+				{
+					// Increment the TimeSinceUsed threshold by our set amount
+					Entry->TimeSinceUsed += Settings.TimeDeltaPerUpdate;
+
+					// Check if this entry is now above our flush threshold. 
+					// if it is, then we'll just remove this asset from our cache
+					if (Entry->TimeSinceUsed > Settings.FlushTimeThreshold)
+					{
+						bAddToLRUList = false;
+						FlushSoundFileDataIndex(i);
+					}
+					else
+					{
+						check(bAddToLRUList);
+						// Count this entry as an inactive loaded sound
+						NumInactiveSounds++;
+					}
+				}
+			}
+			else
+			{
+				// Make sure this is currently active
+				check(Entry->TimeSinceUsed == 0.0f);
+
+				// We count this entry then as an active loaded sound
+				NumActiveSounds++;
+				check(bAddToLRUList);
+			}
+
+			// If we didn't flush this entry, then add it to our time-sorted list
+			if (bAddToLRUList)
+			{
+				NumBytesLoaded += Entry->BulkData.Num();
+
+				// Now add this entry to our sorted active sounds queue
+				LeastRecentlyUsedSoundFiles.HeapPush(FSortedSoundFileEntry(i, Entry->TimeSinceUsed), Predicate);
+			}
+		}
+
+		// Make sure our counts of loaded files and our LRU match
+		check((NumSoundFilesStreamed + NumSoundFilesLoaded) == LeastRecentlyUsedSoundFiles.Num());
+
+		// If our total number of bytes loaded is greater than our cache limit, try to flush any unused sound files
+		// Starting with the least recently used (at the top). If we fail to flush under our limit, that means we 
+		// have so many actively playing/referenced sounds. In that case, log a warning.
+		if (NumBytesLoaded > Settings.TargetMemoryLimit)
+		{
+			bool bFlushSucceed = false;
+			for (int32 i = 0; i < LeastRecentlyUsedSoundFiles.Num(); ++i)
+			{
+				FSortedSoundFileEntry& SortedEntry = LeastRecentlyUsedSoundFiles[i];
+				TSharedPtr<FSoundFileDataEntry> DataEntry = SoundFileData[SortedEntry.Index];
+
+				check(DataEntry.IsValid());
+
+				// Flush the sound file, this will decrement the NumBytesLoaded value
+				int32 MemoryFlushed = FlushSoundFileDataIndex(SortedEntry.Index);
+				NumBytesLoaded -= MemoryFlushed;
+				
+				// Decrease our count of inactive sounds
+				check(NumInactiveSounds > 0);
+				--NumInactiveSounds;
+
+				// if we're now lower than great!
+				if (NumBytesLoaded < Settings.TargetMemoryLimit)
+				{
+					bFlushSucceed = true;
+					break;
+				}
+			}
+
+			// If we failed to flush to under our target sound file cache limit
+			// then log it (but only if it was the first time since we failed since the last time we were properly under)
+			if (!bFlushSucceed && bLogOverMemoryTarget)
+			{
+				bLogOverMemoryTarget = false;
+				UE_LOG(LogUnrealAudio, Warning, TEXT("Audio sound file memory (%d bytes) is over target memory limit (%d)."), NumBytesLoaded, Settings.TargetMemoryLimit);
+			}
+		}
+		else
+		{
+			// We have less bytes loaded than our cache size, so now we can log again if
+			// fail at flushing below our cache size
+			bLogOverMemoryTarget = true;
+		}
+
+	}
+
+	void FSoundFileManager::ReleaseSoundFileHandle(const FSoundFileHandle& SoundFileHandle)
+	{
+		uint32 DataIndex = SoundFileHandle.GetIndex();
+		TSharedPtr<FSoundFileDataEntry> Entry = SoundFileData[DataIndex];
+		check(Entry.IsValid());
+		check(Entry->NumReferences >= 1);
+		Entry->NumReferences--;
+	}
+
+	int32 FSoundFileManager::FlushSoundFileDataIndex(int32 Index)
+	{
+		TSharedPtr<FSoundFileDataEntry> Entry = SoundFileData[Index];
+		check(Entry.IsValid());
+
+		// Make sure nobody is referencing this sound file before we attempt to flush it
+		check(Entry->NumReferences == 0);
+
+		int32 MemoryFlushed = 0;
+
+		// If we have bulk data than this was a loaded sound
+		if (Entry->bIsStreamed)
+		{
+			check(NumSoundFilesStreamed > 0);
+			--NumSoundFilesStreamed;
+
+			check(NameToStreamedSoundMap.Contains(Entry->SoundFilePath));
+			NameToStreamedSoundMap.Remove(Entry->SoundFilePath);
+
+			// Flushing a streaming entry doesn't flush loaded sound file memory
+		}
+		// otherwise it was a streamed sound
+		else
+		{
+			check(NumSoundFilesLoaded > 0);
+			--NumSoundFilesLoaded;
+
+			MemoryFlushed = Entry->BulkData.Num();
+
+			check(NameToLoadedSoundMap.Contains(Entry->SoundFilePath));
+			NameToLoadedSoundMap.Remove(Entry->SoundFilePath);
+		}
+
+		// Release the entry's sound file handle
+		EntityManager.ReleaseEntity(Entry->SoundFileHandle);
+
+		// Free this entry by setting it to a nullptr
+		SoundFileData[Index] = nullptr;
+
+		return MemoryFlushed;
+	}
+
+	void FSoundFileManager::LoadSoundFileDataEntry(TSharedPtr<FSoundFileDataEntry> DataEntry)
+	{
+		check(DataEntry.IsValid());
+
+		ESoundFileError::Type Error;
+		FString Path = DataEntry->SoundFilePath.GetPlainNameString();
+		if (DataEntry->bIsStreamed)
+		{
+			Error = GetSoundFileInfoFromPath(Path, DataEntry->Description, DataEntry->ChannelMap);
+			if (Error == ESoundFileError::NONE)
+			{
+				DataEntry->SoundFileState.Set((int32)ESoundFileState::STREAMING);
+			}
+		}
+		else
+		{
+			Error = LoadSoundFileFromPath(Path, DataEntry->Description, DataEntry->ChannelMap, DataEntry->BulkData);
+			if (Error == ESoundFileError::NONE)
+			{
+				DataEntry->SoundFileState.Set((int32)ESoundFileState::LOADED);
+			}
+			else
+			{
+				DataEntry->SoundFileState.Set((int32)ESoundFileState::HAS_ERROR);
+			}
+		}
+	}
+
+	TSharedPtr<FSoundFileDataEntry> FSoundFileManager::GetEntry(const FSoundFileHandle& Handle)
+	{
+		if (EntityManager.IsValidEntity(Handle))
+		{
+			uint32 Index = Handle.GetIndex();
+			TSharedPtr<FSoundFileDataEntry> Entry = SoundFileData[Index];
+			check(Entry->SoundFileHandle.Id == Handle.Id);
+			return Entry;
+		}
+		return nullptr;
+	}
+
+	void FSoundFileManager::LogSoundFileMemoryInfo() const
+	{
+		UE_LOG(LogUnrealAudio, Display, TEXT("====== Sound file memory usage info ======"), NumBytesLoaded);
+		UE_LOG(LogUnrealAudio, Display, TEXT("Bytes Loaded: %d (%.2f mb), percentage: %.2f"),
+			   NumBytesLoaded, 
+			   (float)NumBytesLoaded / (1024 * 1024), 
+			   (float)NumBytesLoaded / Settings.TargetMemoryLimit
+			);
+		
+		UE_LOG(LogUnrealAudio, Display, TEXT("Num Sound Files Loaded: %d"), NumSoundFilesLoaded);
+		UE_LOG(LogUnrealAudio, Display, TEXT("Num Sound Files Streamed: %d"), NumSoundFilesStreamed);
+		UE_LOG(LogUnrealAudio, Display, TEXT("Num Active Sounds: %d"), NumActiveSounds);
+		UE_LOG(LogUnrealAudio, Display, TEXT("Num Inactive Sounds: %d"), NumInactiveSounds);
+
+		UE_LOG(LogUnrealAudio, Display, TEXT("    Loaded Sounds: "));
+
+		uint32 NumEntries = NumSoundFilesStreamed + NumSoundFilesLoaded;
+		uint32 EntryCount = 0;
+		for (int32 i = 0; i < SoundFileData.Num() && EntryCount < NumEntries; ++i)
+		{
+			if (SoundFileData[i].IsValid())
+			{
+				EntryCount++;
+				TSharedPtr<FSoundFileDataEntry> Entry = SoundFileData[i];
+				UE_LOG(LogUnrealAudio, Display, TEXT("    ------------------"), i);
+				UE_LOG(LogUnrealAudio, Display, TEXT("    EntryIndex: %d"), i);
+				UE_LOG(LogUnrealAudio, Display, TEXT("    Handle: %d"), Entry->SoundFileHandle.Id);
+				UE_LOG(LogUnrealAudio, Display, TEXT("    Bytes: %d (%.2f mb)"), Entry->BulkData.Num(), (float)Entry->BulkData.Num() / (1024*1024));
+				UE_LOG(LogUnrealAudio, Display, TEXT("    Streamed: %s"), (Entry->bIsStreamed ? TEXT("YES") : TEXT("NO")));
+				UE_LOG(LogUnrealAudio, Display, TEXT("    NumReferences: %d"), Entry->NumReferences);
+				UE_LOG(LogUnrealAudio, Display, TEXT("    TimeSinceUsed: %.2f"), Entry->TimeSinceUsed);
+				UE_LOG(LogUnrealAudio, Display, TEXT("    NumFrames: %d"), Entry->Description.NumFrames);
+				UE_LOG(LogUnrealAudio, Display, TEXT("    NumChannels: %d"), Entry->Description.NumChannels);
+				UE_LOG(LogUnrealAudio, Display, TEXT("    SampleRate: %d"), Entry->Description.SampleRate);
+			}
+		}
+
+		UE_LOG(LogUnrealAudio, Display, TEXT("=========================================="), NumBytesLoaded);
+	}
+
+	/************************************************************************/
+	/* UnrealAudioModule Implementations									*/
+	/************************************************************************/
+
+	TSharedPtr<ISoundFile> FUnrealAudioModule::LoadSoundFile(const FName& Path, bool bLoadAsync)
+	{
+		return SoundFileManager.LoadSoundFile(Path, bLoadAsync);
+	}
+
+	TSharedPtr<ISoundFile> FUnrealAudioModule::LoadSoundFile(const FName& Name, TArray<uint8>& InBulkData)
+	{
+		return SoundFileManager.LoadSoundFile(Name, InBulkData);
+	}
+
+	TSharedPtr<ISoundFile> FUnrealAudioModule::StreamSoundFile(const FName& Path, bool bLoadAsync)
+	{
+		return SoundFileManager.StreamSoundFile(Path, bLoadAsync);
+	}
+
+	FSoundFileManager& FUnrealAudioModule::GetSoundFileManager()
+	{
+		return SoundFileManager;
+	}
+
+	int32 FUnrealAudioModule::GetNumSoundFilesLoaded() const
+	{
+		return SoundFileManager.NumSoundFilesLoaded;
+	}
+
+	int32 FUnrealAudioModule::GetNumSoundFilesStreamed() const
+	{
+		return SoundFileManager.NumSoundFilesStreamed;
+	}
+
+	int32 FUnrealAudioModule::GetSoundFileNumBytes() const
+	{
+		MainThreadChecker.CheckThread();
+		return SoundFileManager.NumBytesLoaded;
+	}
+
+	float FUnrealAudioModule::GetSoundFilePercentageOfTargetMemoryLimit() const
+	{
+		MainThreadChecker.CheckThread();
+		return (float) SoundFileManager.NumBytesLoaded / SoundFileManager.Settings.TargetMemoryLimit;
+	}
+
+	void FUnrealAudioModule::LogSoundFileMemoryInfo() const
+	{
+		SoundFileManager.LogSoundFileMemoryInfo();
+	}
+
+
+}
+
+
+
+#endif // #if ENABLE_UNREAL_AUDIO
+
diff --git a/Engine/Source/Runtime/UnrealAudio/Private/UnrealAudioSoundFileManager.h b/Engine/Source/Runtime/UnrealAudio/Private/UnrealAudioSoundFileManager.h
new file mode 100644
--- /dev/null
+++ b/Engine/Source/Runtime/UnrealAudio/Private/UnrealAudioSoundFileManager.h
@@ -0,0 +1,172 @@
+// Copyright 1998-2015 Epic Games, Inc. All Rights Reserved.
+
+#pragma once 
+
+#include "UnrealAudioModule.h"
+#include "UnrealAudioPrivate.h"
+#include "UnrealAudioTypes.h"
+#include "UnrealAudioSoundFile.h"
+#include "UnrealAudioSoundFileInternal.h"
+#include "UnrealAudioEntityManager.h"
+
+#if ENABLE_UNREAL_AUDIO
+
+namespace UAudio
+{
+	class FAsyncSoundFileLoadTask : public FNonAbandonableTask
+	{
+	public:
+		FAsyncSoundFileLoadTask(FUnrealAudioModule* InAudioModule, FSoundFileHandle InSoundFileHandle, const FString& InPath);
+		~FAsyncSoundFileLoadTask();
+
+		void DoWork();
+		bool CanAbandon() { return false; }
+		void Abandon() {}
+
+		FORCEINLINE TStatId GetStatId() const
+		{
+			RETURN_QUICK_DECLARE_CYCLE_STAT(FAsyncSoundFileLoadTask, STATGROUP_ThreadPoolAsyncTasks);
+		}
+
+		FUnrealAudioModule* AudioModule;
+		FSoundFileHandle SoundFileHandle;
+		FString Path;
+		friend class FAsyncTask<FAsyncSoundFileLoadTask>;
+	};
+
+	struct FSoundFileManagerSettings
+	{
+		int32			MaxNumberOfLoadedSounds;
+		int32			NumLoadingThreads;
+		int32			TargetMemoryLimit;
+		float			FlushTimeThreshold;
+		float			TimeDeltaPerUpdate;
+		EThreadPriority LoadingThreadPriority;
+	};
+
+	struct FSoundFileDataEntry
+	{
+		class FSoundFileManager*			SoundFileManager;
+		FName								SoundFilePath;
+		FSoundFileHandle					SoundFileHandle;
+		FSoundFileDescription				Description;
+		TArray<ESoundFileChannelMap::Type>	ChannelMap;
+		TArray<uint8>						BulkData;
+		FThreadSafeCounter					SoundFileState;
+		int32								NumReferences;
+		ESoundFileError::Type				Error;
+		float								TimeSinceUsed;
+		bool								bIsStreamed;
+
+		FSoundFileDataEntry(class FSoundFileManager* InSoundFileManager, const FSoundFileHandle& InSoundFileHandle, bool bInIsStreamed);
+		~FSoundFileDataEntry();
+	};
+
+	class FSoundFileManager
+	{
+	public:
+		FSoundFileManager(FUnrealAudioModule* InAudioModule);
+		~FSoundFileManager();
+
+		void Init(FSoundFileManagerSettings& InSettings);
+		void Shutdown();
+
+		TSharedPtr<ISoundFile> CreateSoundFile();
+
+		/**
+		 * Returns a new ISoundFile shared ptr from the given sound file handle. 
+		 * If the SoundFileHandle is invalid, the returned shared ptr will not be valid.
+		 * @param SoundFileHandle	The input sound file handle.
+		 * @return					TSharedPtr of ISoundFile.
+		 */
+		TSharedPtr<ISoundFile> GetSoundFile(const FSoundFileHandle& SoundFileHandle);
+
+		/** 
+		 * Asynchronously loads a sound file into memory from the given bulk data array.
+		 */
+		TSharedPtr<ISoundFile> LoadSoundFile(const FName& Name, TArray<uint8>& InBulkData);
+
+		/** 
+		 * Asynchronously loads a sound file into memory from the given file path.
+		 */
+		TSharedPtr<ISoundFile> LoadSoundFile(const FName& Path, bool bLoadAsync);
+
+		/** 
+		 * Creates a streaming sound file (doesn't load into memory but does parse file header)
+		 */
+		TSharedPtr<ISoundFile> StreamSoundFile(const FName& Path, bool bLoadAsync);
+
+		ESoundFileState::Type GetState(const FSoundFileHandle& SoundFileHandle) const;
+
+		void ReleaseSoundFileHandle(const FSoundFileHandle& SoundFileHandle);
+
+		void LogSoundFileMemoryInfo() const;
+
+		void Update();
+
+	private:
+
+		int32 FlushSoundFileDataIndex(int32 Index);
+		TSharedPtr<ISoundFile> CreateNewSoundFile(const FName& Path, bool bIsStreamed, bool bLoadAsync);
+		void LoadSoundFileDataEntry(TSharedPtr<FSoundFileDataEntry> DataEntry);
+		TSharedPtr<FSoundFileDataEntry> GetEntry(const FSoundFileHandle& Handle);
+
+		FEntityManager EntityManager;
+		FUnrealAudioModule* AudioModule;
+		FSoundFileManagerSettings Settings;
+
+		FQueuedThreadPool* FileLoadingThreadPool;
+
+		TArray<FSoundFileHandle>				SoundFileHandles;
+		TArray<TSharedPtr<FSoundFileDataEntry>>	SoundFileData;
+
+		// Sorted by last time active (greatest "last used" value on top)
+		TMap<FName, FSoundFileHandle>			NameToLoadedSoundMap;
+		TMap<FName, FSoundFileHandle>			NameToStreamedSoundMap;
+		int32									NumSoundFilesLoaded;
+		int32									NumSoundFilesStreamed;
+		int32									NumBytesLoaded;
+		int32									NumActiveSounds;
+		int32									NumInactiveSounds;
+		bool									bLogOverMemoryTarget;
+
+		FThreadChecker MainThreadChecker;
+
+		struct FSortedSoundFileEntry
+		{
+			// Index of this entry in the sound file data array
+			int32 Index;
+
+			float TimeSinceUsed;
+
+			FSortedSoundFileEntry(int32 InIndex, float InTimeSinceUsed)
+				: Index(InIndex)
+				, TimeSinceUsed(InTimeSinceUsed)
+			{}
+		};
+
+		/** A sound file entry compare predicate. */
+		struct SoundFileEntryComparePredicate
+		{
+			bool operator()(const FSortedSoundFileEntry& A, const FSortedSoundFileEntry& B) const
+			{
+				return A.TimeSinceUsed > B.TimeSinceUsed;
+			}
+		};
+
+		TArray<FSortedSoundFileEntry> LeastRecentlyUsedSoundFiles;
+
+		friend class FAsyncSoundFileLoadTask;
+		friend struct FSoundFileDataEntry;
+		friend class FUnrealAudioModule;
+		friend class FSoundFile;
+	};
+
+
+
+}
+
+
+
+#endif // #if ENABLE_UNREAL_AUDIO
+
