Unreal Audio WIP

diff --git a/Engine/Source/Runtime/UnrealAudio/Private/UnrealAudioSoundFileReader.cpp b/Engine/Source/Runtime/UnrealAudio/Private/UnrealAudioSoundFileReader.cpp
new file mode 100644
--- /dev/null
+++ b/Engine/Source/Runtime/UnrealAudio/Private/UnrealAudioSoundFileReader.cpp
@@ -0,0 +1,1038 @@
+#include "UnrealAudioPrivate.h"
+#include "UnrealAudioSoundFile.h"
+#include "UnrealAudioSoundFileInternal.h"
+#include "UnrealAudioModule.h"
+#include "ModuleManager.h"
+
+#if ENABLE_UNREAL_AUDIO
+
+namespace UAudio
+{
+	// Virtual Sound File Function Pointers
+	typedef SoundFileCount(*VirtualSoundFileGetLengthFuncPtr)(void* UserData);
+	typedef SoundFileCount(*VirtualSoundFileSeekFuncPtr)(SoundFileCount Offset, int32 Mode, void* UserData);
+	typedef SoundFileCount(*VirtualSoundFileReadFuncPtr)(void* DataPtr, SoundFileCount ByteCount, void* UserData);
+	typedef SoundFileCount(*VirtualSoundFileWriteFuncPtr)(const void* DataPtr, SoundFileCount ByteCount, void* UserData);
+	typedef SoundFileCount(*VirtualSoundFileTellFuncPtr)(void* UserData);
+
+	// Struct describing function pointers to call for virtual file IO
+	struct FVirtualSoundFileCallbackInfo
+	{
+		VirtualSoundFileGetLengthFuncPtr VirtualSoundFileGetLength;
+		VirtualSoundFileSeekFuncPtr VirtualSoundFileSeek;
+		VirtualSoundFileReadFuncPtr VirtualSoundFileRead;
+		VirtualSoundFileWriteFuncPtr VirtualSoundFileWrite;
+		VirtualSoundFileTellFuncPtr VirtualSoundFileTell;
+	};
+
+	// SoundFile Constants
+	static const int32 SET_ENCODING_QUALITY = 0x1300;
+	static const int32 SET_CHANNEL_MAP_INFO = 0x1101;
+	static const int32 GET_CHANNEL_MAP_INFO = 0x1100;
+
+	// Exported SoundFile Functions
+	typedef LibSoundFileHandle*(*SoundFileOpenFuncPtr)(const char* Path, int32 Mode, FSoundFileDescription* Description);
+	typedef LibSoundFileHandle*(*SoundFileOpenVirtualFuncPtr)(FVirtualSoundFileCallbackInfo* VirtualFileDescription, int32 Mode, FSoundFileDescription* Description, void* UserData);
+	typedef int32(*SoundFileCloseFuncPtr)(LibSoundFileHandle* FileHandle);
+	typedef int32(*SoundFileErrorFuncPtr)(LibSoundFileHandle* FileHandle);
+	typedef const char*(*SoundFileStrErrorFuncPtr)(LibSoundFileHandle* FileHandle);
+	typedef const char*(*SoundFileErrorNumberFuncPtr)(int32 ErrorNumber);
+	typedef int32(*SoundFileCommandFuncPtr)(LibSoundFileHandle* FileHandle, int32 Command, void* Data, int32 DataSize);
+	typedef int32(*SoundFileFormatCheckFuncPtr)(const FSoundFileDescription* Description);
+	typedef SoundFileCount(*SoundFileSeekFuncPtr)(LibSoundFileHandle* FileHandle, SoundFileCount NumFrames, int32 SeekMode);
+	typedef const char*(*SoundFileGetVersionFuncPtr)(void);
+	typedef SoundFileCount(*SoundFileReadFramesFloatFuncPtr)(LibSoundFileHandle* FileHandle, float* Buffer, SoundFileCount NumFrames);
+	typedef SoundFileCount(*SoundFileReadFramesDoubleFuncPtr)(LibSoundFileHandle* FileHandle, double* Buffer, SoundFileCount NumFrames);
+	typedef SoundFileCount(*SoundFileWriteFramesFloatFuncPtr)(LibSoundFileHandle* FileHandle, const float* Buffer, SoundFileCount NumFrames);
+	typedef SoundFileCount(*SoundFileWriteFramesDoubleFuncPtr)(LibSoundFileHandle* FileHandle, const double* Buffer, SoundFileCount NumFrames);
+	typedef SoundFileCount(*SoundFileReadSamplesFloatFuncPtr)(LibSoundFileHandle* FileHandle, float* Buffer, SoundFileCount NumSamples);
+	typedef SoundFileCount(*SoundFileReadSamplesDoubleFuncPtr)(LibSoundFileHandle* FileHandle, double* Buffer, SoundFileCount NumSamples);
+	typedef SoundFileCount(*SoundFileWriteSamplesFloatFuncPtr)(LibSoundFileHandle* FileHandle, const float* Buffer, SoundFileCount NumSamples);
+	typedef SoundFileCount(*SoundFileWriteSamplesDoubleFuncPtr)(LibSoundFileHandle* FileHandle, const double* Buffer, SoundFileCount NumSamples);
+
+	SoundFileOpenFuncPtr SoundFileOpen = nullptr;
+	SoundFileOpenVirtualFuncPtr SoundFileOpenVirtual = nullptr;
+	SoundFileCloseFuncPtr SoundFileClose = nullptr;
+	SoundFileErrorFuncPtr SoundFileError = nullptr;
+	SoundFileStrErrorFuncPtr SoundFileStrError = nullptr;
+	SoundFileErrorNumberFuncPtr SoundFileErrorNumber = nullptr;
+	SoundFileCommandFuncPtr SoundFileCommand = nullptr;
+	SoundFileFormatCheckFuncPtr SoundFileFormatCheck = nullptr;
+	SoundFileSeekFuncPtr SoundFileSeek = nullptr;
+	SoundFileGetVersionFuncPtr SoundFileGetVersion = nullptr;
+	SoundFileReadFramesFloatFuncPtr SoundFileReadFramesFloat = nullptr;
+	SoundFileReadFramesDoubleFuncPtr SoundFileReadFramesDouble = nullptr;
+	SoundFileWriteFramesFloatFuncPtr SoundFileWriteFramesFloat = nullptr;
+	SoundFileWriteFramesDoubleFuncPtr SoundFileWriteFramesDouble = nullptr;
+	SoundFileReadSamplesFloatFuncPtr SoundFileReadSamplesFloat = nullptr;
+	SoundFileReadSamplesDoubleFuncPtr SoundFileReadSamplesDouble = nullptr;
+	SoundFileWriteSamplesFloatFuncPtr SoundFileWriteSamplesFloat = nullptr;
+	SoundFileWriteSamplesDoubleFuncPtr SoundFileWriteSamplesDouble = nullptr;
+
+	/**
+	Function implementations of virtual function callbacks
+	*/
+
+	static SoundFileCount OnSoundFileGetLengthBytes(void* UserData)
+	{
+		SoundFileCount Length = 0;
+		((ISoundFileParser*)UserData)->GetLengthBytes(Length);
+		return Length;
+	}
+
+	static SoundFileCount OnSoundFileSeekBytes(SoundFileCount Offset, int32 Mode, void* UserData)
+	{
+		SoundFileCount OutOffset = 0;
+		((ISoundFileParser*)UserData)->SeekBytes(Offset, (ESoundFileSeekMode::Type)Mode, OutOffset);
+		return OutOffset;
+	}
+
+	static SoundFileCount OnSoundFileReadBytes(void* DataPtr, SoundFileCount ByteCount, void* UserData)
+	{
+		SoundFileCount OutBytesRead = 0;
+		((ISoundFileParser*)UserData)->ReadBytes(DataPtr, ByteCount, OutBytesRead);
+		return OutBytesRead;
+	}
+
+	static SoundFileCount OnSoundFileWriteBytes(const void* DataPtr, SoundFileCount ByteCount, void* UserData)
+	{
+		SoundFileCount OutBytesWritten = 0;
+		((ISoundFileParser*)UserData)->WriteBytes(DataPtr, ByteCount, OutBytesWritten);
+		return OutBytesWritten;
+	}
+
+	static SoundFileCount OnSoundFileTell(void* UserData)
+	{
+		SoundFileCount OutOffset = 0;
+		((ISoundFileParser*)UserData)->GetOffsetBytes(OutOffset);
+		return OutOffset;
+	}
+
+	/**
+	* Gets the default channel mapping for the given channel number
+	*/
+	static void GetDefaultMappingsForChannelNumber(int32 NumChannels, TArray<ESoundFileChannelMap::Type>& ChannelMap)
+	{
+		check(ChannelMap.Num() == NumChannels);
+
+		switch (NumChannels)
+		{
+		case 1:	// MONO
+		ChannelMap[0] = ESoundFileChannelMap::MONO;
+		break;
+
+		case 2:	// STEREO
+		ChannelMap[0] = ESoundFileChannelMap::LEFT;
+		ChannelMap[1] = ESoundFileChannelMap::RIGHT;
+		break;
+
+		case 3:	// 2.1
+		ChannelMap[0] = ESoundFileChannelMap::LEFT;
+		ChannelMap[1] = ESoundFileChannelMap::RIGHT;
+		ChannelMap[2] = ESoundFileChannelMap::LFE;
+		break;
+
+		case 4: // Quadraphonic
+		ChannelMap[0] = ESoundFileChannelMap::LEFT;
+		ChannelMap[1] = ESoundFileChannelMap::RIGHT;
+		ChannelMap[2] = ESoundFileChannelMap::BACK_LEFT;
+		ChannelMap[3] = ESoundFileChannelMap::BACK_RIGHT;
+		break;
+
+		case 5: // 5.0
+		ChannelMap[0] = ESoundFileChannelMap::LEFT;
+		ChannelMap[1] = ESoundFileChannelMap::RIGHT;
+		ChannelMap[2] = ESoundFileChannelMap::CENTER;
+		ChannelMap[3] = ESoundFileChannelMap::SIDE_LEFT;
+		ChannelMap[4] = ESoundFileChannelMap::SIDE_RIGHT;
+		break;
+
+		case 6: // 5.1
+		ChannelMap[0] = ESoundFileChannelMap::LEFT;
+		ChannelMap[1] = ESoundFileChannelMap::RIGHT;
+		ChannelMap[2] = ESoundFileChannelMap::CENTER;
+		ChannelMap[3] = ESoundFileChannelMap::LFE;
+		ChannelMap[4] = ESoundFileChannelMap::SIDE_LEFT;
+		ChannelMap[5] = ESoundFileChannelMap::SIDE_RIGHT;
+		break;
+
+		case 7: // 6.1
+		ChannelMap[0] = ESoundFileChannelMap::LEFT;
+		ChannelMap[1] = ESoundFileChannelMap::RIGHT;
+		ChannelMap[2] = ESoundFileChannelMap::CENTER;
+		ChannelMap[3] = ESoundFileChannelMap::LFE;
+		ChannelMap[4] = ESoundFileChannelMap::SIDE_LEFT;
+		ChannelMap[5] = ESoundFileChannelMap::SIDE_RIGHT;
+		ChannelMap[6] = ESoundFileChannelMap::BACK_CENTER;
+		break;
+
+		case 8: // 7.1
+		ChannelMap[0] = ESoundFileChannelMap::LEFT;
+		ChannelMap[1] = ESoundFileChannelMap::RIGHT;
+		ChannelMap[2] = ESoundFileChannelMap::CENTER;
+		ChannelMap[3] = ESoundFileChannelMap::LFE;
+		ChannelMap[4] = ESoundFileChannelMap::BACK_LEFT;
+		ChannelMap[5] = ESoundFileChannelMap::BACK_RIGHT;
+		ChannelMap[6] = ESoundFileChannelMap::SIDE_LEFT;
+		ChannelMap[7] = ESoundFileChannelMap::SIDE_RIGHT;
+		break;
+
+		default:
+		break;
+		}
+	}
+
+	static ESoundFileError::Type GetSoundDesriptionInternal(LibSoundFileHandle** OutFileHandle, const FString& FilePath, FSoundFileDescription& OutputDescription, TArray<ESoundFileChannelMap::Type>& OutChannelMap)
+	{
+		*OutFileHandle = nullptr;
+
+		// Check to see if the file exists
+		if (!FPaths::FileExists(FilePath))
+		{
+			UE_LOG(LogUnrealAudio, Error, TEXT("Sound file %s doesn't exist."), *FilePath);
+			return ESoundFileError::FILE_DOESNT_EXIST;
+		}
+
+		// open a sound file handle to get the description
+		*OutFileHandle = SoundFileOpen(TCHAR_TO_ANSI(*FilePath), ESoundFileOpenMode::READING, &OutputDescription);
+		if (!*OutFileHandle)
+		{
+			FString StrError = FString(SoundFileStrError(nullptr));
+			UE_LOG(LogUnrealAudio, Error, TEXT("Failed to open sound file %s: %s"), *FilePath, *StrError);
+			return ESoundFileError::FAILED_TO_OPEN;
+		}
+
+		// Try to get a channel mapping
+		int32 NumChannels = OutputDescription.NumChannels;
+		OutChannelMap.Init(ESoundFileChannelMap::INVALID, NumChannels);
+
+		int32 Result = SoundFileCommand(*OutFileHandle, GET_CHANNEL_MAP_INFO, (int32*)OutChannelMap.GetData(), sizeof(int32)*NumChannels);
+
+		// If we failed to get the file's channel map definition, then we set the default based on the number of channels
+		if (Result == 0)
+		{
+			GetDefaultMappingsForChannelNumber(NumChannels, OutChannelMap);
+		}
+		else
+		{
+			// Check to see if the channel map we did get back is filled with INVALID channels
+			bool bIsInvalid = false;
+			for (ESoundFileChannelMap::Type ChannelType : OutChannelMap)
+			{
+				if (ChannelType == ESoundFileChannelMap::INVALID)
+				{
+					bIsInvalid = true;
+					break;
+				}
+			}
+			// If invalid, then we need to get the default channel mapping
+			if (bIsInvalid)
+			{
+				GetDefaultMappingsForChannelNumber(NumChannels, OutChannelMap);
+			}
+		}
+
+		return ESoundFileError::NONE;
+	}
+
+	/************************************************************************/
+	/* FSoundFileReader														*/
+	/************************************************************************/
+
+	FSoundFileReader::FSoundFileReader(FUnrealAudioModule* InAudioModule)
+		: AudioModule(InAudioModule)
+		, CurrentIndexBytes(0)
+		, FileHandle(nullptr)
+		, State(ESoundFileState::UNINITIALIZED)
+		, CurrentError(ESoundFileError::NONE)
+	{
+	}
+
+	FSoundFileReader::~FSoundFileReader()
+	{
+		Release();
+		check(FileHandle == nullptr);
+	}
+
+	ESoundFileError::Type FSoundFileReader::GetLengthBytes(SoundFileCount& OutLength) const
+	{
+		if (!SoundFileData.IsValid())
+		{
+			return ESoundFileError::INVALID_DATA;
+		}
+
+		int32 DataSize;
+		ESoundFileError::Type Error = SoundFileData->GetDataSize(DataSize);
+		if (Error == ESoundFileError::NONE)
+		{
+			OutLength = DataSize;
+			return ESoundFileError::NONE;
+		}
+		return Error;
+	}
+
+	ESoundFileError::Type FSoundFileReader::SeekBytes(SoundFileCount Offset, ESoundFileSeekMode::Type SeekMode, SoundFileCount& OutOffset)
+	{
+		if (!SoundFileData.IsValid())
+		{
+			return ESoundFileError::INVALID_DATA;
+		}
+ 
+		int32 DataSize;
+		ESoundFileError::Type Error = SoundFileData->GetDataSize(DataSize);
+		if (Error != ESoundFileError::NONE)
+		{
+			return Error;
+		}
+
+		SoundFileCount MaxBytes = DataSize;
+		if (MaxBytes == 0)
+		{
+			OutOffset = 0;
+			CurrentIndexBytes = 0;
+			return ESoundFileError::NONE;
+		}
+ 
+		switch (SeekMode)
+		{
+			case ESoundFileSeekMode::FROM_START:
+				CurrentIndexBytes = Offset;
+				break;
+
+			case ESoundFileSeekMode::FROM_CURRENT:
+				CurrentIndexBytes += Offset;
+				break;
+
+			case ESoundFileSeekMode::FROM_END:
+				CurrentIndexBytes = MaxBytes + Offset;
+				break;
+
+			default:
+				checkf(false, TEXT("Uknown seek mode!"));
+				break;
+		}
+
+		// Wrap the byte index to fall between 0 and MaxBytes
+		while (CurrentIndexBytes < 0)
+		{
+			CurrentIndexBytes += MaxBytes;
+		}
+
+		while (CurrentIndexBytes > MaxBytes)
+		{
+			CurrentIndexBytes -= MaxBytes;
+		}
+
+		OutOffset = CurrentIndexBytes;
+		return ESoundFileError::NONE;
+	}
+
+	ESoundFileError::Type FSoundFileReader::ReadBytes(void* DataPtr, SoundFileCount NumBytes, SoundFileCount& OutNumBytesRead)
+	{
+		if (!SoundFileData.IsValid())
+		{
+			return ESoundFileError::INVALID_DATA;
+		}
+
+		SoundFileCount EndByte = CurrentIndexBytes + NumBytes;
+
+		int32 DataSize;
+		ESoundFileError::Type Error = SoundFileData->GetDataSize(DataSize);
+		if (Error != ESoundFileError::NONE)
+		{
+			return Error;
+		}
+		SoundFileCount MaxBytes = DataSize;
+		if (EndByte >= MaxBytes)
+		{
+			NumBytes = MaxBytes - CurrentIndexBytes;
+		}
+
+		if (NumBytes > 0)
+		{
+			TArray<uint8>* BulkData = nullptr;
+			Error = SoundFileData->GetBulkData(&BulkData);
+			if (Error != ESoundFileError::NONE)
+			{
+				return Error;
+			}
+
+			check(BulkData != nullptr);
+
+			FMemory::Memcpy(DataPtr, (const void*)&(*BulkData)[CurrentIndexBytes], NumBytes);
+			CurrentIndexBytes += NumBytes;
+		}
+		OutNumBytesRead = NumBytes;
+		return ESoundFileError::NONE;
+	}
+
+	ESoundFileError::Type FSoundFileReader::WriteBytes(const void* DataPtr, SoundFileCount NumBytes, SoundFileCount& OutNumBytesWritten)
+	{
+		// This should never get called in the reader class
+		check(false);
+		return ESoundFileError::NONE;
+	}
+
+	ESoundFileError::Type FSoundFileReader::GetOffsetBytes(SoundFileCount& OutOffset) const
+	{
+		OutOffset = CurrentIndexBytes;
+		return ESoundFileError::NONE;
+	}
+
+	ESoundFileError::Type FSoundFileReader::Init(TSharedPtr<ISoundFile> InSoundFileData, bool bIsStreamed)
+	{
+		if (bIsStreamed)
+		{
+			return InitStreamed(InSoundFileData);
+		}
+		else
+		{
+			return InitLoaded(InSoundFileData);
+		}
+	}
+
+	ESoundFileError::Type FSoundFileReader::Release()
+	{
+		if (FileHandle)
+		{
+			SoundFileClose(FileHandle);
+			FileHandle = nullptr;
+		}
+		return ESoundFileError::NONE;
+	}
+
+	ESoundFileError::Type FSoundFileReader::SeekFrames(SoundFileCount Offset, ESoundFileSeekMode::Type SeekMode, SoundFileCount& OutOffset)
+	{
+		SoundFileCount Pos = SoundFileSeek(FileHandle, Offset, (int32)SeekMode);
+		if (Pos == -1)
+		{
+			FString StrErr = SoundFileStrError(FileHandle);
+			UE_LOG(LogUnrealAudio, Error, TEXT("Failed to seek file: %s"), *StrErr);
+			return SetError(ESoundFileError::FAILED_TO_SEEK);
+		}
+		return ESoundFileError::NONE;
+	}
+
+	ESoundFileError::Type FSoundFileReader::ReadFrames(float* DataPtr, SoundFileCount NumFrames, SoundFileCount& OutNumFramesRead)
+	{
+		OutNumFramesRead = SoundFileReadFramesFloat(FileHandle, DataPtr, NumFrames);
+		return ESoundFileError::NONE;
+	}
+
+	ESoundFileError::Type FSoundFileReader::ReadFrames(double* DataPtr, SoundFileCount NumFrames, SoundFileCount& OutNumFramesRead)
+	{
+		OutNumFramesRead = SoundFileReadFramesDouble(FileHandle, DataPtr, NumFrames);
+		return ESoundFileError::NONE;
+	}
+
+	ESoundFileError::Type FSoundFileReader::ReadSamples(float* DataPtr, SoundFileCount NumSamples, SoundFileCount& OutNumSamplesRead)
+	{
+		OutNumSamplesRead = SoundFileReadSamplesFloat(FileHandle, DataPtr, NumSamples);
+		return ESoundFileError::NONE;
+	}
+
+	ESoundFileError::Type FSoundFileReader::ReadSamples(double* DataPtr, SoundFileCount NumSamples, SoundFileCount& OutNumSamplesRead)
+	{
+		OutNumSamplesRead = SoundFileReadSamplesDouble(FileHandle, DataPtr, NumSamples);
+		return ESoundFileError::NONE;
+	}
+
+	ESoundFileError::Type FSoundFileReader::InitLoaded(TSharedPtr<ISoundFile> InSoundFileData)
+	{
+		if (!(State.GetValue() == ESoundFileState::UNINITIALIZED || State.GetValue() == ESoundFileState::LOADING))
+		{
+			return SetError(ESoundFileError::ALREADY_INTIALIZED);
+		}
+
+		check(InSoundFileData.IsValid());
+ 		check(FileHandle == nullptr);
+
+		// Setting sound file data initializes this sound file
+		SoundFileData = InSoundFileData;
+		check(SoundFileData.IsValid());
+
+		bool bIsStreamed;
+		ESoundFileError::Type Error = SoundFileData->IsStreamed(bIsStreamed);
+		if (Error != ESoundFileError::NONE)
+		{
+			return Error;
+		}
+
+		if (bIsStreamed)
+		{
+			return ESoundFileError::INVALID_DATA;
+		}
+
+		ESoundFileState::Type SoundFileState;
+		Error = SoundFileData->GetState(SoundFileState);
+		if (Error != ESoundFileError::NONE)
+		{
+			return Error;
+		}
+
+		if (SoundFileState != ESoundFileState::LOADED)
+		{
+			return ESoundFileError::INVALID_STATE;
+		}
+
+ 		// Open up a virtual file handle with this data
+		FVirtualSoundFileCallbackInfo VirtualSoundFileInfo;
+		VirtualSoundFileInfo.VirtualSoundFileGetLength	= OnSoundFileGetLengthBytes;
+		VirtualSoundFileInfo.VirtualSoundFileSeek		= OnSoundFileSeekBytes;
+		VirtualSoundFileInfo.VirtualSoundFileRead		= OnSoundFileReadBytes;
+		VirtualSoundFileInfo.VirtualSoundFileWrite		= OnSoundFileWriteBytes;
+		VirtualSoundFileInfo.VirtualSoundFileTell		= OnSoundFileTell;
+ 
+		FSoundFileDescription Description;
+		SoundFileData->GetDescription(Description);
+ 		if (!SoundFileFormatCheck(&Description))
+ 		{
+ 			return SetError(ESoundFileError::INVALID_INPUT_FORMAT);
+ 		}
+ 
+ 		FileHandle = SoundFileOpenVirtual(&VirtualSoundFileInfo, ESoundFileOpenMode::READING, &Description, (void*)this);
+ 		if (!FileHandle)
+ 		{
+ 			FString StrErr = SoundFileStrError(nullptr);
+ 			UE_LOG(LogUnrealAudio, Error, TEXT("Failed to intitialize sound file: %s"), *StrErr);
+ 			return SetError(ESoundFileError::FAILED_TO_OPEN);
+ 		}
+ 
+ 		State.Set(ESoundFileState::INITIALIZED);
+		return ESoundFileError::NONE;
+	}
+
+	ESoundFileError::Type FSoundFileReader::InitStreamed(TSharedPtr<ISoundFile> InSoundFileData)
+	{
+		if (!(State.GetValue() == ESoundFileState::UNINITIALIZED || State.GetValue() == ESoundFileState::LOADING))
+		{
+			return SetError(ESoundFileError::ALREADY_INTIALIZED);
+		}
+
+		check(InSoundFileData.IsValid());
+		check(FileHandle == nullptr);
+
+		// Setting sound file data initializes this sound file
+		SoundFileData = InSoundFileData;
+		check(SoundFileData.IsValid());
+
+		bool bIsStreamed;
+		ESoundFileError::Type Error = SoundFileData->IsStreamed(bIsStreamed);
+		if (Error != ESoundFileError::NONE)
+		{
+			return Error;
+		}
+
+		if (!bIsStreamed)
+		{
+			return ESoundFileError::INVALID_DATA;
+		}
+
+		ESoundFileState::Type SoundFileState;
+		Error = SoundFileData->GetState(SoundFileState);
+		if (Error != ESoundFileError::NONE)
+		{
+			return Error;
+		}
+
+		if (SoundFileState != ESoundFileState::STREAMING)
+		{
+			return ESoundFileError::INVALID_STATE;
+		}
+
+		FName NamePath;
+		Error = SoundFileData->GetPath(NamePath);
+		if (Error != ESoundFileError::NONE)
+		{
+			return Error;
+		}
+
+		FString FilePath = NamePath.GetPlainNameString();
+
+		FSoundFileDescription Description;
+		TArray<ESoundFileChannelMap::Type> ChannelMap;
+		Error = GetSoundDesriptionInternal(&FileHandle, FilePath, Description, ChannelMap);
+		if (Error == ESoundFileError::NONE)
+		{
+			// Tell this reader that we're in streaming mode.
+			State.Set(ESoundFileState::STREAMING);
+			return ESoundFileError::NONE;
+		}
+		else
+		{
+			return SetError(Error);
+		}
+	}
+
+	ESoundFileError::Type FSoundFileReader::SetError(ESoundFileError::Type InError)
+	{
+		if (InError != ESoundFileError::NONE)
+		{
+			State.Set(ESoundFileState::HAS_ERROR);
+		}
+		CurrentError.Set(InError);
+		return InError;
+	}
+
+
+	/************************************************************************/
+	/* FSoundFileWriter														*/
+	/************************************************************************/
+
+	FSoundFileWriter::FSoundFileWriter(FUnrealAudioModule* InAudioModule)
+		: AudioModule(InAudioModule)
+		, CurrentIndexBytes(0)
+		, FileHandle(nullptr)
+		, EncodingQuality(0.0)
+		, State(ESoundFileState::UNINITIALIZED)
+		, CurrentError(ESoundFileError::NONE)
+	{
+
+	}
+
+	FSoundFileWriter::~FSoundFileWriter()
+	{
+
+	}
+
+	ESoundFileError::Type FSoundFileWriter::GetLengthBytes(SoundFileCount& OutLength) const
+	{
+		OutLength = BulkData.Num();
+		return ESoundFileError::NONE;
+	}
+
+	ESoundFileError::Type FSoundFileWriter::SeekBytes(SoundFileCount Offset, ESoundFileSeekMode::Type SeekMode, SoundFileCount& OutOffset)
+	{
+		int32 DataSize = BulkData.Num();
+
+		SoundFileCount MaxBytes = DataSize;
+		if (MaxBytes == 0)
+		{
+			OutOffset = 0;
+			CurrentIndexBytes = 0;
+			return ESoundFileError::NONE;
+		}
+
+		switch (SeekMode)
+		{
+			case ESoundFileSeekMode::FROM_START:
+				CurrentIndexBytes = Offset;
+				break;
+
+			case ESoundFileSeekMode::FROM_CURRENT:
+				CurrentIndexBytes += Offset;
+				break;
+
+			case ESoundFileSeekMode::FROM_END:
+				CurrentIndexBytes = MaxBytes + Offset;
+				break;
+
+			default:
+				checkf(false, TEXT("Uknown seek mode!"));
+				break;
+		}
+
+		// Wrap the byte index to fall between 0 and MaxBytes
+		while (CurrentIndexBytes < 0)
+		{
+			CurrentIndexBytes += MaxBytes;
+		}
+
+		while (CurrentIndexBytes > MaxBytes)
+		{
+			CurrentIndexBytes -= MaxBytes;
+		}
+
+		OutOffset = CurrentIndexBytes;
+		return ESoundFileError::NONE;
+	}
+
+	ESoundFileError::Type FSoundFileWriter::ReadBytes(void* DataPtr, SoundFileCount NumBytes, SoundFileCount& OutNumBytesRead)
+	{
+		// This shouldn't get called in the writer
+		check(false);
+		return ESoundFileError::NONE;
+	}
+
+	ESoundFileError::Type FSoundFileWriter::WriteBytes(const void* DataPtr, SoundFileCount NumBytes, SoundFileCount& OutNumBytesWritten)
+	{
+		const uint8* InDataBytes = (const uint8*)DataPtr;
+
+		SoundFileCount MaxBytes = BulkData.Num();
+
+		// Append the data to the buffer
+		if (CurrentIndexBytes == MaxBytes)
+		{
+			BulkData.Append(InDataBytes, NumBytes);
+			CurrentIndexBytes += NumBytes;
+		}
+		else if ((CurrentIndexBytes + NumBytes) < MaxBytes)
+		{
+			// Write over the existing data in the buffer
+			for (SoundFileCount i = 0; i < NumBytes; ++i)
+			{
+				BulkData[CurrentIndexBytes++] = InDataBytes[i];
+			}
+		}
+		else
+		{
+			// Overwrite some data until end of current buffer size
+			SoundFileCount RemainingBytes = MaxBytes - CurrentIndexBytes;
+			SoundFileCount InputBufferIndex = 0;
+			for (InputBufferIndex = 0; InputBufferIndex < RemainingBytes; ++InputBufferIndex)
+			{
+				BulkData[CurrentIndexBytes++] = InDataBytes[InputBufferIndex];
+			}
+
+			// Now append the remainder of the input buffer to the internal data buffer
+			const uint8* RemainingData = &InDataBytes[InputBufferIndex];
+			RemainingBytes = NumBytes - RemainingBytes;
+			BulkData.Append(RemainingData, RemainingBytes);
+			CurrentIndexBytes += RemainingBytes;
+			check(CurrentIndexBytes == BulkData.Num());
+		}
+		OutNumBytesWritten = NumBytes;
+		return ESoundFileError::NONE;
+	}
+
+	ESoundFileError::Type FSoundFileWriter::GetOffsetBytes(SoundFileCount& OutOffset) const
+	{
+		OutOffset = CurrentIndexBytes;
+		return ESoundFileError::NONE;
+	}
+
+	ESoundFileError::Type FSoundFileWriter::Init(const FSoundFileDescription& InDescription, const TArray<ESoundFileChannelMap::Type>& InChannelMap, double InEncodingQuality)
+	{
+		State.Set(ESoundFileState::INITIALIZED);
+
+		BulkData.Reset();
+		Description		= InDescription;
+		ChannelMap		= InChannelMap;
+		EncodingQuality = InEncodingQuality;
+
+		// First check the input format to make sure it's valid
+		if (!SoundFileFormatCheck(&InDescription))
+		{
+			UE_LOG(LogUnrealAudio,
+				   Error,
+				   TEXT("Sound file input format (%s - %s) is invalid."),
+				   ESoundFileFormat::ToStringMajor(InDescription.FormatFlags),
+				   ESoundFileFormat::ToStringMinor(InDescription.FormatFlags));
+
+			return SetError(ESoundFileError::INVALID_INPUT_FORMAT);
+		}
+
+		// Make sure we have the right number of channels and our channel map size
+		if (InChannelMap.Num() != InDescription.NumChannels)
+		{
+			UE_LOG(LogUnrealAudio, Error, TEXT("Channel map didn't match the input NumChannels"));
+			return SetError(ESoundFileError::INVALID_CHANNEL_MAP);
+		}
+
+		FVirtualSoundFileCallbackInfo VirtualSoundFileInfo;
+		VirtualSoundFileInfo.VirtualSoundFileGetLength	= OnSoundFileGetLengthBytes;
+		VirtualSoundFileInfo.VirtualSoundFileSeek		= OnSoundFileSeekBytes;
+		VirtualSoundFileInfo.VirtualSoundFileRead		= OnSoundFileReadBytes;
+		VirtualSoundFileInfo.VirtualSoundFileWrite		= OnSoundFileWriteBytes;
+		VirtualSoundFileInfo.VirtualSoundFileTell		= OnSoundFileTell;
+
+		FSoundFileDescription Description = InDescription;
+		FileHandle = SoundFileOpenVirtual(&VirtualSoundFileInfo, ESoundFileOpenMode::WRITING, &Description, (void*)this);
+		if (!FileHandle)
+		{
+			FString StrErr = FString(SoundFileStrError(nullptr));
+			UE_LOG(LogUnrealAudio, Error, TEXT("Failed to open empty sound file: %s"), *StrErr);
+			return SetError(ESoundFileError::FAILED_TO_OPEN);
+		}
+
+		int32 Result = SoundFileCommand(FileHandle, SET_CHANNEL_MAP_INFO, (int32*)InChannelMap.GetData(), sizeof(int32)*Description.NumChannels);
+		if (Result)
+		{
+			FString StrErr = SoundFileStrError(nullptr);
+			UE_LOG(LogUnrealAudio, Error, TEXT("Failed to set the channel map on empty file for writing: %s"), *StrErr);
+			return SetError(ESoundFileError::INVALID_CHANNEL_MAP);
+		}
+
+		if ((Description.FormatFlags & ESoundFileFormat::MAJOR_FORMAT_MASK) == ESoundFileFormat::OGG)
+		{
+			int32 Result = SoundFileCommand(FileHandle, SET_ENCODING_QUALITY, &EncodingQuality, sizeof(double));
+			if (Result != 1)
+			{
+				FString StrErr = SoundFileStrError(FileHandle);
+				UE_LOG(LogUnrealAudio, Error, TEXT("Failed to set encoding quality: %s"), *StrErr);
+				return SetError(ESoundFileError::BAD_ENCODING_QUALITY);
+			}
+		}
+
+		return ESoundFileError::NONE;
+	}
+
+	ESoundFileError::Type FSoundFileWriter::Release()
+	{
+		if (FileHandle)
+		{
+			SoundFileClose(FileHandle);
+			FileHandle = nullptr;
+		}
+		return ESoundFileError::NONE;
+	}
+
+	ESoundFileError::Type FSoundFileWriter::SeekFrames(SoundFileCount Offset, ESoundFileSeekMode::Type SeekMode, SoundFileCount& OutOffset)
+	{
+		SoundFileCount Pos = SoundFileSeek(FileHandle, Offset, (int32)SeekMode);
+		if (Pos == -1)
+		{
+			FString StrErr = SoundFileStrError(FileHandle);
+			UE_LOG(LogUnrealAudio, Error, TEXT("Failed to seek file: %s"), *StrErr);
+			return SetError(ESoundFileError::FAILED_TO_SEEK);
+		}
+		return ESoundFileError::NONE;
+	}
+
+	ESoundFileError::Type FSoundFileWriter::WriteFrames(const float* DataPtr, SoundFileCount NumFrames, SoundFileCount& OutNumFramesWritten)
+	{
+		OutNumFramesWritten = SoundFileWriteFramesFloat(FileHandle, DataPtr, NumFrames);
+		return ESoundFileError::NONE;
+	}
+
+	ESoundFileError::Type FSoundFileWriter::WriteFrames(const double* DataPtr, SoundFileCount NumFrames, SoundFileCount& OutNumFramesWritten)
+	{
+		OutNumFramesWritten = SoundFileWriteFramesDouble(FileHandle, DataPtr, NumFrames);
+		return ESoundFileError::NONE;
+	}
+
+	ESoundFileError::Type FSoundFileWriter::WriteSamples(const float* DataPtr, SoundFileCount NumSamples, SoundFileCount& OutNumSampleWritten)
+	{
+		OutNumSampleWritten = SoundFileWriteSamplesFloat(FileHandle, DataPtr, NumSamples);
+		return ESoundFileError::NONE;
+	}
+
+	ESoundFileError::Type FSoundFileWriter::WriteSamples(const double* DataPtr, SoundFileCount NumSamples, SoundFileCount& OutNumSampleWritten)
+	{
+		OutNumSampleWritten = SoundFileWriteSamplesDouble(FileHandle, DataPtr, NumSamples);
+		return ESoundFileError::NONE;
+	}
+
+	ESoundFileError::Type FSoundFileWriter::GetData(TArray<uint8>** OutBulkData)
+	{
+		*OutBulkData = &BulkData;
+		return ESoundFileError::NONE;
+	}
+
+	ESoundFileError::Type FSoundFileWriter::SetError(ESoundFileError::Type InError)
+	{
+		if (InError != ESoundFileError::NONE)
+		{
+			State.Set(ESoundFileState::HAS_ERROR);
+		}
+		CurrentError.Set(InError);
+		return InError;
+	}
+
+	/************************************************************************/
+	/* FUnrealAudioModule													*/
+	/************************************************************************/
+
+	static void* GetSoundFileDllHandle()
+	{
+		void* DllHandle = nullptr;
+#if PLATFORM_WINDOWS
+		FString Path = FPaths::EngineDir() / FString(TEXT("Binaries/ThirdParty/libsndfile/Win64/"));
+		FPlatformProcess::PushDllDirectory(*Path);
+		DllHandle = FPlatformProcess::GetDllHandle(*(Path + "libsndfile-1.dll"));
+		FPlatformProcess::PopDllDirectory(*Path);
+#elif PLATFORM_MAC
+		//		FString Path = FPaths::EngineDir() / FString(TEXT("Binaries/ThirdParty/libsndfile/Mac/"));
+		//		FPlatformProcess::PushDllDirectory(*Path);
+		DllHandle = FPlatformProcess::GetDllHandle(TEXT("libsndfile.1.dylib"));
+		//		FPlatformProcess::PopDllDirectory(*Path);
+#endif
+		return DllHandle;
+	}
+
+	bool FUnrealAudioModule::LoadSoundFileLib()
+	{
+		SoundFileDllHandle = GetSoundFileDllHandle();
+		if (!SoundFileDllHandle)
+		{
+			UE_LOG(LogUnrealAudio, Error, TEXT("Failed to load Sound File dll"));
+			return false;
+		}
+
+		bool bSuccess = true;
+
+		void* LambdaDLLHandle = SoundFileDllHandle;
+
+		// Helper function to load DLL exports and report warnings
+		auto GetSoundFileDllExport = [&LambdaDLLHandle](const TCHAR* FuncName, bool& bInSuccess) -> void*
+		{
+			if (bInSuccess)
+			{
+				void* FuncPtr = FPlatformProcess::GetDllExport(LambdaDLLHandle, FuncName);
+				if (FuncPtr == nullptr)
+				{
+					bInSuccess = false;
+					UE_LOG(LogUnrealAudio, Warning, TEXT("Failed to locate the expected DLL import function '%s' in the SoundFile DLL."), *FuncName);
+					FPlatformProcess::FreeDllHandle(LambdaDLLHandle);
+					LambdaDLLHandle = nullptr;
+				}
+				return FuncPtr;
+			}
+			else
+			{
+				return nullptr;
+			}
+		};
+
+		SoundFileOpen = (SoundFileOpenFuncPtr)GetSoundFileDllExport(TEXT("sf_open"), bSuccess);
+		SoundFileOpenVirtual = (SoundFileOpenVirtualFuncPtr)GetSoundFileDllExport(TEXT("sf_open_virtual"), bSuccess);
+		SoundFileClose = (SoundFileCloseFuncPtr)GetSoundFileDllExport(TEXT("sf_close"), bSuccess);
+		SoundFileError = (SoundFileErrorFuncPtr)GetSoundFileDllExport(TEXT("sf_error"), bSuccess);
+		SoundFileStrError = (SoundFileStrErrorFuncPtr)GetSoundFileDllExport(TEXT("sf_strerror"), bSuccess);
+		SoundFileErrorNumber = (SoundFileErrorNumberFuncPtr)GetSoundFileDllExport(TEXT("sf_error_number"), bSuccess);
+		SoundFileCommand = (SoundFileCommandFuncPtr)GetSoundFileDllExport(TEXT("sf_command"), bSuccess);
+		SoundFileFormatCheck = (SoundFileFormatCheckFuncPtr)GetSoundFileDllExport(TEXT("sf_format_check"), bSuccess);
+		SoundFileSeek = (SoundFileSeekFuncPtr)GetSoundFileDllExport(TEXT("sf_seek"), bSuccess);
+		SoundFileGetVersion = (SoundFileGetVersionFuncPtr)GetSoundFileDllExport(TEXT("sf_version_string"), bSuccess);
+		SoundFileReadFramesFloat = (SoundFileReadFramesFloatFuncPtr)GetSoundFileDllExport(TEXT("sf_readf_float"), bSuccess);
+		SoundFileReadFramesDouble = (SoundFileReadFramesDoubleFuncPtr)GetSoundFileDllExport(TEXT("sf_readf_double"), bSuccess);
+		SoundFileWriteFramesFloat = (SoundFileWriteFramesFloatFuncPtr)GetSoundFileDllExport(TEXT("sf_writef_float"), bSuccess);
+		SoundFileWriteFramesDouble = (SoundFileWriteFramesDoubleFuncPtr)GetSoundFileDllExport(TEXT("sf_writef_double"), bSuccess);
+		SoundFileReadSamplesFloat = (SoundFileReadSamplesFloatFuncPtr)GetSoundFileDllExport(TEXT("sf_read_float"), bSuccess);
+		SoundFileReadSamplesDouble = (SoundFileReadSamplesDoubleFuncPtr)GetSoundFileDllExport(TEXT("sf_read_double"), bSuccess);
+		SoundFileWriteSamplesFloat = (SoundFileWriteSamplesFloatFuncPtr)GetSoundFileDllExport(TEXT("sf_write_float"), bSuccess);
+		SoundFileWriteSamplesDouble = (SoundFileWriteSamplesDoubleFuncPtr)GetSoundFileDllExport(TEXT("sf_write_double"), bSuccess);
+
+		// make sure we're successful
+		check(bSuccess);
+		return bSuccess;
+	}
+
+	bool FUnrealAudioModule::ShutdownSoundFileLib()
+	{
+		if (SoundFileDllHandle)
+		{
+			FPlatformProcess::FreeDllHandle(SoundFileDllHandle);
+			SoundFileDllHandle = nullptr;
+			SoundFileOpen = nullptr;
+			SoundFileOpenVirtual = nullptr;
+			SoundFileClose = nullptr;
+			SoundFileError = nullptr;
+			SoundFileStrError = nullptr;
+			SoundFileErrorNumber = nullptr;
+			SoundFileCommand = nullptr;
+			SoundFileFormatCheck = nullptr;
+			SoundFileSeek = nullptr;
+			SoundFileGetVersion = nullptr;
+			SoundFileReadFramesFloat = nullptr;
+			SoundFileReadFramesDouble = nullptr;
+			SoundFileWriteFramesFloat = nullptr;
+			SoundFileWriteFramesDouble = nullptr;
+			SoundFileReadSamplesFloat = nullptr;
+			SoundFileReadSamplesDouble = nullptr;
+			SoundFileWriteSamplesFloat = nullptr;
+			SoundFileWriteSamplesDouble = nullptr;
+		}
+		return true;
+	}
+
+	TSharedPtr<FSoundFileReader> FUnrealAudioModule::CreateSoundFileReader()
+	{
+		return TSharedPtr<FSoundFileReader>(new FSoundFileReader(this));
+	}
+
+	TSharedPtr<FSoundFileWriter> FUnrealAudioModule::CreateSoundFileWriter()
+	{
+		return TSharedPtr<FSoundFileWriter>(new FSoundFileWriter(this));
+	}
+
+	/************************************************************************/
+	/* Exported Functions													*/
+	/************************************************************************/
+
+	bool GetListOfSoundFilesInDirectory(const FString& Directory, TArray<FString>& OutSoundFileList, FString* TypeFilter /*= nullptr */)
+	{
+		return false;
+	}
+
+	bool GetSoundFileDescription(const FString& FilePath, FSoundFileDescription& OutputDescription, TArray<ESoundFileChannelMap::Type>& OutChannelMap)
+	{
+		LibSoundFileHandle* FileHandle = nullptr;
+		ESoundFileError::Type Error = GetSoundDesriptionInternal(&FileHandle, FilePath, OutputDescription, OutChannelMap);
+		if (Error == ESoundFileError::NONE)
+		{
+			DEBUG_AUDIO_CHECK(FileHandle != nullptr);
+			SoundFileClose(FileHandle);
+			return true;
+		}
+		return false;
+	}
+
+	bool GetSoundFileDescription(const FString& FilePath, FSoundFileDescription& OutputDescription)
+	{
+		TArray<ESoundFileChannelMap::Type> OutChannelMap;
+		return GetSoundFileDescription(FilePath, OutputDescription, OutChannelMap);
+	}
+
+	bool GetFileExtensionForFormatFlags(int32 FormatFlags, FString& OutExtension)
+	{
+		if (FormatFlags & ESoundFileFormat::OGG)
+		{
+			OutExtension = TEXT("ogg");
+		}
+		else if (FormatFlags & ESoundFileFormat::WAV)
+		{
+			OutExtension = TEXT("wav");
+		}
+		else if (FormatFlags & ESoundFileFormat::AIFF)
+		{
+			OutExtension = TEXT("aiff");
+		}
+		else if (FormatFlags & ESoundFileFormat::FLAC)
+		{
+			OutExtension = TEXT("flac");
+		}
+		else
+		{
+			return false;
+		}
+
+		return true;
+	}
+
+	ESoundFileError::Type GetSoundFileInfoFromPath(const FString& FilePath, FSoundFileDescription& Description, TArray<ESoundFileChannelMap::Type>& ChannelMap)
+	{
+		// Load the description and channel map info
+		LibSoundFileHandle* FileHandle = nullptr;
+		ESoundFileError::Type Error = GetSoundDesriptionInternal(&FileHandle, FilePath, Description, ChannelMap);
+		if (FileHandle)
+		{
+			SoundFileClose(FileHandle);
+		}
+		return Error;
+	}
+
+	ESoundFileError::Type LoadSoundFileFromPath(const FString& FilePath, FSoundFileDescription& Description, TArray<ESoundFileChannelMap::Type>& ChannelMap, TArray<uint8>& BulkData)
+	{
+		ESoundFileError::Type Error = GetSoundFileInfoFromPath(FilePath, Description, ChannelMap);
+		if (Error != ESoundFileError::NONE)
+		{
+			return Error;
+		}
+
+		// Now read the data from disk into the bulk data array
+		if (FFileHelper::LoadFileToArray(BulkData, *FilePath))
+		{
+			return ESoundFileError::NONE;
+		}
+		else
+		{
+			return ESoundFileError::FAILED_TO_LOAD_BYTE_DATA;
+		}
+	}
+
+
+
+}
+
+#endif // #if ENABLE_UNREAL_AUDIO
\ No newline at end of file
diff --git a/Engine/Source/Runtime/UnrealAudio/Private/UnrealAudioThread.h b/Engine/Source/Runtime/UnrealAudio/Private/UnrealAudioThread.h
new file mode 100644
--- /dev/null
+++ b/Engine/Source/Runtime/UnrealAudio/Private/UnrealAudioThread.h
@@ -0,0 +1,18 @@
+// Copyright 1998-2015 Epic Games, Inc. All Rights Reserved.
+
+#pragma once 
+
+#include "UnrealAudioModule.h"
+
+
+#if ENABLE_UNREAL_AUDIO
+
+namespace UAudio
+{
+
+
+}
+
+#endif // #if ENABLE_UNREAL_AUDIO
+
+
diff --git a/Engine/Source/Runtime/UnrealAudio/Private/UnrealAudioUtilities.cpp b/Engine/Source/Runtime/UnrealAudio/Private/UnrealAudioUtilities.cpp
new file mode 100644
--- /dev/null
+++ b/Engine/Source/Runtime/UnrealAudio/Private/UnrealAudioUtilities.cpp
@@ -0,0 +1,213 @@
+// Copyright 1998-2015 Epic Games, Inc. All Rights Reserved.
+
+#include "UnrealAudioPrivate.h"
+#include "UnrealAudioUtilities.h"
+
+#if ENABLE_UNREAL_AUDIO
+
+// Useful debug check macros
+
+namespace UAudio
+{
+	FSafeFloat::FSafeFloat(float InitValue)
+	{
+		int32* IntVal = reinterpret_cast<int32*>(&InitValue);
+		SafeValue.Set(*IntVal);
+	}
+
+	void FSafeFloat::Set(float Value)
+	{
+		int32* IntVal = reinterpret_cast<int32*>(&Value);
+		SafeValue.Set(*IntVal);
+	}
+
+	float FSafeFloat::Get() const
+	{
+		int32 Value = SafeValue.GetValue();
+		float* FloatValue = reinterpret_cast<float*>(&Value);
+		return *FloatValue;
+	}
+
+	FCommandData::FCommandData()
+		: DataType(ECommandData::INVALID)
+	{
+		Data.PtrVal = nullptr;
+	}
+
+	FCommandData::FCommandData(void* Val)
+		: DataType(ECommandData::POINTER)
+	{
+		Data.PtrVal = Val;
+	}
+
+	FCommandData::FCommandData(const FEntityHandle& EntityHandle)
+		: DataType(ECommandData::HANDLE)
+	{
+		Data.Handle.Id = EntityHandle.Id;
+	}
+
+	FCommandData::FCommandData(float Val)
+		: DataType(ECommandData::FLOAT_32)
+	{
+		Data.Float32Val = Val;
+	}
+
+	FCommandData::FCommandData(double Val)
+		: DataType(ECommandData::FLOAT_64)
+	{
+		Data.Float64Val = Val;
+	}
+
+	FCommandData::FCommandData(uint8 Val)
+		: DataType(ECommandData::UINT8)
+	{
+		Data.UnsignedInt8 = Val;
+	}
+
+	FCommandData::FCommandData(uint32 Val)
+		: DataType(ECommandData::UINT32)
+	{
+		Data.UnsignedInt32 = Val;
+	}
+
+	FCommandData::FCommandData(uint64 Val)
+		: DataType(ECommandData::UINT64)
+	{
+		Data.UnsignedInt64 = Val;
+	}
+
+	FCommandData::FCommandData(bool Val)
+		: DataType(ECommandData::BOOL)
+	{
+		Data.BoolVal = Val;
+	}
+
+	FCommandData::FCommandData(int32 Val)
+		: DataType(ECommandData::INT32)
+	{
+		Data.Int32Val = Val;
+	}
+
+	FCommandData::FCommandData(int64 Val)
+		: DataType(ECommandData::INT64)
+	{
+		Data.Int64Val = Val;
+	}
+
+	FCommand::FCommand()
+		: Id(0)
+		, NumArguments(0)
+	{
+	}
+
+	FCommand::FCommand(uint32 InCommand)
+		: Id(InCommand)
+		, NumArguments(0)
+	{
+	}
+
+	FCommand::FCommand(uint32 InCommand, const FCommandData& Arg0)
+		: Id(InCommand)
+		, NumArguments(1)
+	{
+		Arguments[0] = Arg0;
+	}
+
+	FCommand::FCommand(uint32 InCommand, const FCommandData& Arg0, const FCommandData& Arg1)
+		: Id(InCommand)
+		, NumArguments(2)
+	{
+		Arguments[0] = Arg0;
+		Arguments[1] = Arg1;
+	}
+
+	FCommand::FCommand(uint32 InCommand, const FCommandData& Arg0, const FCommandData& Arg1, const FCommandData& Arg2)
+		: Id(InCommand)
+		, NumArguments(3)
+	{
+		Arguments[0] = Arg0;
+		Arguments[1] = Arg1;
+		Arguments[2] = Arg2;
+	}
+
+	FCommand::FCommand(uint32 InCommand, const FCommandData& Arg0, const FCommandData& Arg1, const FCommandData& Arg2, const FCommandData& Arg3)
+		: Id(InCommand)
+		, NumArguments(4)
+	{
+		Arguments[0] = Arg0;
+		Arguments[1] = Arg1;
+		Arguments[2] = Arg2;
+		Arguments[3] = Arg3;
+	}
+
+	FCommand::FCommand(uint32 InCommand, const FCommandData& Arg0, const FCommandData& Arg1, const FCommandData& Arg2, const FCommandData& Arg3, const FCommandData& Arg4)
+		: Id(InCommand)
+		, NumArguments(5)
+	{
+		Arguments[0] = Arg0;
+		Arguments[1] = Arg1;
+		Arguments[2] = Arg2;
+		Arguments[3] = Arg3;
+		Arguments[4] = Arg4;
+	}
+
+	float FDynamicParamData::Compute(uint32 Index, float CurrentTimeSec)
+	{
+		float Result = 0.0f;
+		if (bIsDone[Index])
+		{
+			Result = CurrentValue[Index];
+		}
+		else
+		{
+			float Fraction = (float)FMath::Min((CurrentTimeSec - StartTime[Index]) / DeltaTime[Index], 1.0);
+			bIsDone[Index] = Fraction >= 1.0f;
+
+			// LERP the volume parameter
+			// TODO: other interpolations
+			Result = (1.0f - Fraction) * StartValue[Index] + Fraction * EndValue[Index];
+			CurrentValue[Index] = Result;
+		}
+		return Result;
+	}
+
+	void FDynamicParamData::Init(uint32 NumElements)
+	{
+		StartValue.Init(1.0f, NumElements);
+		EndValue.Init(1.0f, NumElements);
+		CurrentValue.Init(1.0f, NumElements);
+		StartTime.Init(0.0f, NumElements);
+		DeltaTime.Init(0.0f, NumElements);
+		bIsDone.Init(true, NumElements);
+	}
+
+	void FDynamicParamData::InitEntry(uint32 Index)
+	{
+		DEBUG_AUDIO_CHECK(Index < (uint32)StartValue.Num());
+
+		StartValue[Index] = 1.0f;
+		EndValue[Index] = 1.0f;
+		CurrentValue[Index] = 1.0f;
+		StartTime[Index] = 0.0f;
+		DeltaTime[Index] = 0.0f;
+		bIsDone[Index] = true;
+	}
+
+	void FDynamicParamData::SetValue(uint32 Index, float InValue, float InStartTime, float InDeltaTimeSec)
+	{
+		if (InValue != CurrentValue[Index])
+		{
+			StartValue[Index] = CurrentValue[Index];
+			EndValue[Index] = InValue;
+			StartTime[Index] = InStartTime;
+			DeltaTime[Index] = InDeltaTimeSec;
+			bIsDone[Index] = false;
+		}
+	}
+
+
+}
+
+#endif // #if ENABLE_UNREAL_AUDIO
+
+
diff --git a/Engine/Source/Runtime/UnrealAudio/Private/UnrealAudioUtilities.h b/Engine/Source/Runtime/UnrealAudio/Private/UnrealAudioUtilities.h
new file mode 100644
--- /dev/null
+++ b/Engine/Source/Runtime/UnrealAudio/Private/UnrealAudioUtilities.h
@@ -0,0 +1,403 @@
+// Copyright 1998-2015 Epic Games, Inc. All Rights Reserved.
+
+#pragma once 
+
+#include "UnrealAudioPrivate.h"
+#include "UnrealAudioModule.h"
+#include "UnrealAudioEntityManager.h"
+
+#if ENABLE_UNREAL_AUDIO
+
+// Useful debug check macros
+
+#if ENABLE_UNREAL_AUDIO_EXTRA_DEBUG_CHECKS
+#define DEBUG_AUDIO_CHECK_AUDIO_THREAD(AUDIO_MODULE)	(AUDIO_MODULE->CheckAudioThread())
+#define DEBUG_AUDIO_CHECK_MAIN_THREAD(AUDIO_MODULE)		(AUDIO_MODULE->CheckMainThread())
+#define DEBUG_AUDIO_CHECK(EXP)							check(EXP)
+#define DEBUG_AUDIO_CHECK_MSG(EXP, MSG)					checkf(EXP, MSG)
+#else // ENABLE_UNREAL_AUDIO_EXTRA_DEBUG_CHECKS
+#define DEBUG_AUDIO_CHECK_AUDIO_THREAD(AUDIO_MODULE)
+#define DEBUG_AUDIO_CHECK_MAIN_THREAD(AUDIO_MODULE)
+#define DEBUG_AUDIO_CHECK(EXP)
+#define DEBUG_AUDIO_CHECK_MSG(EXP, MSG)
+#endif // #else // ENABLE_UNREAL_AUDIO_EXTRA_DEBUG_CHECKS
+
+namespace UAudio
+{
+	/**
+	* TSafeQueue
+	* A simple (minimal lock) thread-safe multi-producer/multi-consumer FIFO queue
+	* Based on H. Sutter's queue described in:
+	* http://www.drdobbs.com/parallel/writing-a-generalized-concurrent-queue/211601363
+	*/
+	template<class ElementType>
+	class TSafeQueue
+	{
+	public:
+		TSafeQueue()
+		{
+			First = new FNode(nullptr);
+			Last = First;
+			ProducerLock = 0;
+			ConsumerLock = 0;
+		}
+
+		~TSafeQueue()
+		{
+			while (First != nullptr)
+			{
+				FNode* Temp = First;
+				First = Temp->Next;
+				delete Temp->Element;
+				delete Temp;
+			}
+		}
+
+		void Enqueue(const ElementType& Element)
+		{
+			FNode* Node = new FNode(new ElementType(Element));
+
+			// Allow multiple threads to enqueue, but only let one do it at a time
+			while (FPlatformAtomics::InterlockedExchange(&ProducerLock, 1))
+			{
+				// spin while acquiring exclusivity. This will only spin if 
+				// a different thread has already set the ProducerLock to 1
+			}
+
+			FPlatformAtomics::InterlockedExchangePtr((void**)&Last->Next, Node);
+			FPlatformAtomics::InterlockedExchangePtr((void**)&Last, Node);
+
+			// Free the lock
+			ProducerLock = 0;
+		}
+
+		bool Dequeue(ElementType& OutElement)
+		{
+			while (FPlatformAtomics::InterlockedExchange(&ConsumerLock, 1))
+			{
+				// spin while acquiring exclusivity. This will only spin if 
+				// a different thread has already set the ConsumerLock to 1
+			}
+
+			FNode* Head = First;
+			FNode* Next = First->Next;
+			if (Next != nullptr)
+			{
+				ElementType* Element = Next->Element;
+				Next->Element = nullptr;
+				First = Next;
+				ConsumerLock = 0;
+				OutElement = *Element;
+				delete Element;
+				delete Head;
+				return true;
+			}
+
+			ConsumerLock = 0;
+			return false;
+		}
+
+	private:
+
+		// Align the size of the FNode struct to ensure that two 
+		// FNode objects won't occupy the same cache line (to avoid "cache-line ping-ponging")
+		struct MS_ALIGN(CACHE_LINE_SIZE) FNode
+		{
+			FNode(ElementType* InElement)
+				: Element(InElement)
+				, Next(nullptr)
+			{}
+
+			ElementType* Element;
+			FNode* Next;
+		};
+
+		MS_ALIGN(CACHE_LINE_SIZE) FNode*			First;
+		MS_ALIGN(CACHE_LINE_SIZE) volatile int32	ConsumerLock;
+		MS_ALIGN(CACHE_LINE_SIZE) FNode*			Last;
+		MS_ALIGN(CACHE_LINE_SIZE) volatile int32	ProducerLock;
+	};
+
+	/**
+	* TLocklessQueue
+	* A thread-safe lockless single-producer/single-consumer FIFO queue
+	* Based on H. Sutter's queue described in:
+	* http://www.drdobbs.com/parallel/writing-lock-free-code-a-corrected-queue/210604448?pgno=3
+	*/
+	template<class ElementType>
+	class TLocklessQueue
+	{
+	public:
+		TLocklessQueue()
+		{
+			First = new FNode(nullptr);
+			Last = First;
+			Divider = First;
+		}
+
+		~TLocklessQueue()
+		{
+			while (First != nullptr)
+			{
+				FNode* Temp = First;
+				First = Temp->NextNode;
+				delete Temp;
+			}
+		}
+
+		void Enqueue(const ElementType& Element)
+		{
+			FNode* Node = new FNode(Element);
+			Last->Next = Node;
+
+			// Publish the new node to the last node, now the consumer thread can read the node
+			FPlatformAtomics::InterlockedExchangePtr((void**)&Last, Last->Next);
+
+			// Trim unused nodes lazily
+			while (First != Divider)
+			{
+				FNode* Temp = First;
+				First = First->Next;
+				delete Temp;
+			}
+		}
+
+		bool Dequeue(ElementType& OutElement)
+		{
+			if (Divider != Last)
+			{
+				OutElement = Divider->Next->Element;
+				// Publish that we took this node
+				FPlatformAtomics::InterlockedExchangePtr((void**)&Divider, Divider->Next);
+				return true;
+			}
+			return false;
+		}
+
+	private:
+		struct FNode
+		{
+			FNode(const ElementType& InElement)
+			: Element(InElement)
+			, Next(nullptr)
+			{}
+
+			ElementType Element;
+			FNode* Next;
+		};
+
+		MS_ALIGN(CACHE_LINE_SIZE) FNode* First;		// For producer
+		MS_ALIGN(CACHE_LINE_SIZE) FNode* Last;		// Shared
+		MS_ALIGN(CACHE_LINE_SIZE) FNode* Divider;	// Shared
+	};
+
+	class FSafeFloat
+	{
+	public:
+		FSafeFloat(float InitValue);
+		void Set(float Value);
+		float Get() const;
+
+	private:
+		FThreadSafeCounter SafeValue;
+	};
+
+	/**
+	* ECommandData
+	* Enumerated list of types used in audio commands.
+	*/
+	namespace ECommandData
+	{
+		enum Type
+		{
+			INVALID,
+			POINTER,
+			HANDLE,
+			FLOAT_32,
+			FLOAT_64,
+			BOOL,
+			UINT8,
+			UINT32,
+			UINT64,
+			INT32,
+			INT64
+		};
+	}
+
+	/** 
+	* FCommandData
+	* Simple type union for use as arguments in commands.
+	*/
+	struct FCommandData
+	{
+		FCommandData();
+		FCommandData(void* Val);
+		FCommandData(const FEntityHandle& EntityHandle);
+		FCommandData(float Val);
+		FCommandData(double Val);
+		FCommandData(uint8 Val);
+		FCommandData(uint32 Val);
+		FCommandData(uint64 Val);
+		FCommandData(bool Val);
+		FCommandData(int32 Val);
+		FCommandData(int64 Val);
+
+		ECommandData::Type DataType;
+
+		union
+		{
+			void* PtrVal;
+			float Float32Val;
+			double Float64Val;
+			bool BoolVal;
+			uint8 UnsignedInt8;
+			uint32 UnsignedInt32;
+			uint64 UnsignedInt64;
+			int32 Int32Val;
+			int64 Int64Val;
+			struct
+			{
+				FEntityHandle Handle;
+			};
+		} Data;
+	};
+
+	/**
+	* TCommand
+	* A struct used to send commands between threads for audio events.
+	*/
+	struct FCommand
+	{
+		uint32 Id;
+		FCommandData Arguments[5];
+		int32 NumArguments;
+
+		FCommand();
+		FCommand(uint32 InCommand);
+		FCommand(uint32 InCommand, const FCommandData& Arg0);
+		FCommand(uint32 InCommand, const FCommandData& Arg0, const FCommandData& Arg1);
+		FCommand(uint32 InCommand, const FCommandData& Arg0, const FCommandData& Arg1, const FCommandData& Arg2);
+		FCommand(uint32 InCommand, const FCommandData& Arg0, const FCommandData& Arg1, const FCommandData& Arg2, const FCommandData& Arg3);
+		FCommand(uint32 InCommand, const FCommandData& Arg0, const FCommandData& Arg1, const FCommandData& Arg2, const FCommandData& Arg3, const FCommandData& Arg4);
+	};
+
+	template <class CommandType>
+	class TCommandQueue
+	{
+	public:
+		TCommandQueue(uint32 CircularQueueSize)
+			: CircularCommandQueue(CircularQueueSize)
+		{}
+
+		void Enqueue(const CommandType& InCommand)
+		{
+			if (!CircularCommandQueue.Enqueue(InCommand))
+			{
+				// if circular queue is full, then push to fallback queue (which performs allocations)
+				SafeCommandQueue.Enqueue(InCommand);
+			}
+		}
+
+		bool Dequeue(CommandType& OutCommand)
+		{
+			if (CircularCommandQueue.Dequeue(OutCommand))
+			{
+				// Empty the fallback command queue and push to circular
+				CommandType TempCommand;
+				if (SafeCommandQueue.Dequeue(TempCommand))
+				{
+					CircularCommandQueue.Enqueue(TempCommand);
+				}
+				return true;
+			}
+			return false;
+		}
+
+	private:
+		TCircularQueue<CommandType> CircularCommandQueue;
+		TSafeQueue<CommandType> SafeCommandQueue;
+	};
+
+	struct FDynamicParam
+	{
+		float StartValue;
+		float EndValue;
+		float CurrentValue;
+		float StartTime;
+		float DeltaTime;
+		bool bIsDone;
+	};
+
+	struct FDynamicParamData
+	{
+		TArray<float> StartValue;
+		TArray<float> EndValue;
+		TArray<float> CurrentValue;
+		TArray<double> StartTime;
+		TArray<double> DeltaTime;
+		TArray<bool> bIsDone;
+
+		void Init(uint32 NumElements);
+		float Compute(uint32 Index, float CurrentTimeSec);
+		void InitEntry(uint32 Index);
+		void SetValue(uint32 Index, float InValue, float InStartTime, float InDeltaTimeSec);
+	};
+
+	class FThreadChecker
+	{
+	public:
+		FThreadChecker()
+#if ENABLE_UNREAL_AUDIO_THREAD_DEBUGGING
+			: ThreadId(INDEX_NONE)
+#endif
+		{}
+
+		void InitThread()
+		{
+#if ENABLE_UNREAL_AUDIO_THREAD_DEBUGGING
+			ThreadId = FPlatformTLS::GetCurrentThreadId();
+#endif
+		}
+
+		void CheckThread() const
+		{
+#if ENABLE_UNREAL_AUDIO_THREAD_DEBUGGING
+			uint32 CurrentThreadId = FPlatformTLS::GetCurrentThreadId();
+			checkf(CurrentThreadId == ThreadId,
+					TEXT("Function called on wrong thread with id '%d' but supposed to be called on main thread (id=%d)."),
+					CurrentThreadId,
+					ThreadId
+					);
+#endif
+		}
+
+		void CheckThread()
+		{
+#if ENABLE_UNREAL_AUDIO_THREAD_DEBUGGING
+			if (ThreadId == INDEX_NONE)
+			{
+				ThreadId = FPlatformTLS::GetCurrentThreadId();
+			}
+			else
+			{
+				uint32 CurrentThreadId = FPlatformTLS::GetCurrentThreadId();
+				checkf(CurrentThreadId == ThreadId,
+					   TEXT("Function called on wrong thread with id '%d' but supposed to be called on main thread (id=%d)."),
+					   CurrentThreadId,
+					   ThreadId
+					);
+			}
+#endif
+		}
+
+	private:
+#if ENABLE_UNREAL_AUDIO_THREAD_DEBUGGING
+		uint32 ThreadId;
+#endif
+	};
+
+}
+
+#endif // #if ENABLE_UNREAL_AUDIO
+
+
diff --git a/Engine/Source/Runtime/UnrealAudio/Private/UnrealAudioVoice.cpp b/Engine/Source/Runtime/UnrealAudio/Private/UnrealAudioVoice.cpp
new file mode 100644
--- /dev/null
+++ b/Engine/Source/Runtime/UnrealAudio/Private/UnrealAudioVoice.cpp
@@ -0,0 +1,450 @@
+// Copyright 1998-2015 Epic Games, Inc. All Rights Reserved.
+// 
+#include "UnrealAudioPrivate.h"
+#include "UnrealAudioModule.h"
+#include "UnrealAudioVoiceInternal.h"
+
+#define UNREAL_AUDIO_VOICE_COMMAND_QUEUE_SIZE (50)
+
+#if ENABLE_UNREAL_AUDIO
+
+#define AUDIO_VOICE_CHECK_ERROR()					\
+		if (VoiceState == EVoiceState::HAS_ERROR)	\
+		{											\
+			return LastError;						\
+		}
+
+#define AUDIO_VOICE_CHECK_SUSPEND()							\
+		if (PlayingState == EVoicePlayingState::SUSPENDED)	\
+		{													\
+			return EVoiceError::NONE;						\
+		}
+
+
+
+namespace UAudio
+{
+	/************************************************************************/
+	/* FVoice implementation												*/
+	/************************************************************************/
+
+	FVoice::FVoice(class FUnrealAudioModule* InParentModule, const FVoiceInitializationParams& InParams)
+		: AudioModule(InParentModule)
+		, PlayingState(EVoicePlayingState::NOT_PLAYING)
+		, VoiceState(EVoiceState::STOPPED)
+		, Params(InParams)
+		, VoiceDataIndex(INDEX_NONE)
+		, LastError(EVoiceError::NONE)
+		, DurationSeconds(0.0f)
+		, bPlayCalled(false)
+	{
+		check(Params.SoundFile.IsValid());
+
+		// Get the duration right away
+// 		const FSoundFileDescription& Description = Params.SoundFile->GetDescription();
+// 		check(Description.SampleRate > 0);
+// 		DurationSeconds = static_cast<float>(Description.NumFrames) / Description.SampleRate;
+	}
+
+	FVoice::~FVoice()
+	{
+		DEBUG_AUDIO_CHECK(VoiceDataIndex == INDEX_NONE);
+	}
+
+	EVoiceError::Type FVoice::Play()
+	{
+		AUDIO_VOICE_CHECK_ERROR();
+
+		if (VoiceState != EVoiceState::PLAYING)
+		{
+			// Make sure play is not called 2x on the same voice.
+			bPlayCalled = true;
+
+			// Build the initialization data for the play call
+			FVoiceInitializationData* VoiceInitData = new FVoiceInitializationData();
+			VoiceInitData->SoundFile			= Params.SoundFile;
+			VoiceInitData->EmitterHandle		= Params.Emitter.IsValid() ? ((FEmitter*)Params.Emitter.Get())->GetHandle() : FEmitterHandle();
+			VoiceInitData->BaselineVolumeScale	= Params.BaselineVolumeScale;
+			VoiceInitData->DynamicVolumeScale	= PlayInfo.VolumeScale;
+			VoiceInitData->DynamicVolumeTime	= PlayInfo.VolumeScaleTime;
+			VoiceInitData->BaselinePitchScale	= Params.BaselinePitchScale;
+			VoiceInitData->DynamicPitchScale	= PlayInfo.PitchScale;
+			VoiceInitData->DynamicPitchTime		= PlayInfo.PitchScaletime;
+			VoiceInitData->PriorityWeight		= Params.PriorityWeight;
+			VoiceInitData->DurationSeconds		= DurationSeconds;
+
+			VoiceInitData->VoiceFlags = EVoiceFlag::NONE;
+			if (Params.bIsLooping)
+			{
+				VoiceInitData->VoiceFlags |= EVoiceFlag::LOOPING;
+			}
+
+			if (Params.Emitter.IsValid())
+			{
+				VoiceInitData->VoiceFlags |= EVoiceFlag::SPATIALIZED;
+			}
+
+			// Play the voice with the voice manager
+			check(AudioModule != nullptr);
+			FVoiceManager& VoiceManager = AudioModule->GetVoiceManager();
+			EVoiceError::Type Result = VoiceManager.PlayVoice(this, VoiceInitData, VoiceHandle);
+
+			// If nothing went wrong then our VoiceState is "playing"
+			if (Result == EVoiceError::NONE)
+			{
+				VoiceState = EVoiceState::PLAYING;
+			}
+			else
+			{
+				VoiceState = EVoiceState::HAS_ERROR;
+				LastError = Result;
+			}
+			return Result;
+		}
+
+		return EVoiceError::NONE;
+	}
+
+	EVoiceError::Type FVoice::Pause(float InFadeTimeSec /* = 0.0f */)
+	{
+		AUDIO_VOICE_CHECK_ERROR();
+		AUDIO_VOICE_CHECK_SUSPEND();
+
+		if (VoiceState == EVoiceState::PLAYING)
+		{
+			if (InFadeTimeSec == 0.0f)
+			{
+				VoiceState = EVoiceState::PAUSED;
+			}
+			else
+			{
+				VoiceState = EVoiceState::PAUSING;
+			}
+
+			FVoiceManager& VoiceManager = AudioModule->GetVoiceManager();
+			return SetError(VoiceManager.PauseVoice(VoiceHandle, InFadeTimeSec));
+		}
+		else if (VoiceState == EVoiceState::STOPPED)
+		{
+			VoiceState = EVoiceState::PAUSED;
+			if (VoiceDataIndex != INDEX_NONE)
+			{
+				FVoiceManager& VoiceManager = AudioModule->GetVoiceManager();
+				return SetError(VoiceManager.PauseVoice(VoiceHandle, 0.0f));
+			}
+			else
+			{
+				return EVoiceError::NONE;
+			}
+		}
+
+		return EVoiceError::NONE;
+	}
+
+	EVoiceError::Type FVoice::Stop(float InFadeTimeSec /* = 0.0f */)
+	{
+		AUDIO_VOICE_CHECK_ERROR();
+		AUDIO_VOICE_CHECK_SUSPEND();
+
+		if (VoiceState != EVoiceState::STOPPED)
+		{
+			if (InFadeTimeSec == 0.0f)
+			{
+				VoiceState = EVoiceState::STOPPED;
+			}
+			else
+			{
+				VoiceState = EVoiceState::STOPPING;
+			}
+
+			FVoiceManager& VoiceManager = AudioModule->GetVoiceManager();
+			return SetError(VoiceManager.StopVoice(VoiceHandle, InFadeTimeSec));
+		}
+
+		return EVoiceError::NONE;
+	}
+
+	EVoiceError::Type FVoice::AddVoiceListener(IVoiceListener* Listener)
+	{
+		AUDIO_VOICE_CHECK_ERROR();
+		AUDIO_VOICE_CHECK_SUSPEND();
+
+		if (VoiceListeners.Contains(Listener))
+		{
+			return EVoiceError::LISTENER_ALREADY_ADDED;
+		}
+
+		VoiceListeners.Add(Listener);
+		return EVoiceError::NONE;
+	}
+
+	EVoiceError::Type FVoice::RemoveVoiceListener(IVoiceListener* Listener)
+	{
+		AUDIO_VOICE_CHECK_ERROR();
+		AUDIO_VOICE_CHECK_SUSPEND();
+
+		if (!VoiceListeners.Contains(Listener))
+		{
+			return EVoiceError::LISTENER_NOT_ADDED;
+		}
+
+		VoiceListeners.Remove(Listener);
+		return EVoiceError::NONE;
+	}
+
+	EVoiceError::Type FVoice::SetVolumeScale(float InVolumeScale, float InDeltaTimeSec /*= 0.0f */)
+	{
+		AUDIO_VOICE_CHECK_ERROR();
+		AUDIO_VOICE_CHECK_SUSPEND();
+
+		if (PlayInfo.VolumeScale != InVolumeScale || PlayInfo.VolumeScaleTime != InDeltaTimeSec)
+		{
+			PlayInfo.VolumeScale = InVolumeScale;
+			PlayInfo.VolumeScaleTime = InDeltaTimeSec;
+
+			if (IsPlaying())
+			{
+				DEBUG_AUDIO_CHECK(VoiceDataIndex != INDEX_NONE);
+				FVoiceManager& VoiceManager = AudioModule->GetVoiceManager();
+				return SetError(VoiceManager.SetVolumeScale(VoiceDataIndex, InVolumeScale, InDeltaTimeSec));
+			}
+		}
+		return EVoiceError::NONE;
+	}
+
+	EVoiceError::Type FVoice::SetPitchScale(float InPitchScale, float InDeltaTimeSec /* = 0.0f */)
+	{
+		AUDIO_VOICE_CHECK_ERROR();
+		AUDIO_VOICE_CHECK_SUSPEND();
+
+		if (PlayInfo.PitchScale != InPitchScale || PlayInfo.PitchScaletime != InDeltaTimeSec)
+		{
+			PlayInfo.PitchScale = InPitchScale;
+			PlayInfo.PitchScaletime = InDeltaTimeSec;
+			if (IsPlaying())
+			{
+				DEBUG_AUDIO_CHECK(VoiceDataIndex != INDEX_NONE);
+				FVoiceManager& VoiceManager = AudioModule->GetVoiceManager();
+				return SetError(VoiceManager.SetPitchScale(VoiceDataIndex, InPitchScale, InDeltaTimeSec));
+			}
+		}
+		return EVoiceError::NONE;
+	}
+
+	EVoiceError::Type FVoice::GetPitchScale(float& OutPitchScale) const
+	{
+		AUDIO_VOICE_CHECK_ERROR();
+
+		if (IsPlaying())
+		{
+			DEBUG_AUDIO_CHECK(VoiceDataIndex != INDEX_NONE);
+			FVoiceManager& VoiceManager = AudioModule->GetVoiceManager();
+			return VoiceManager.GetPitchScale(VoiceDataIndex, OutPitchScale);
+		}
+		OutPitchScale = 1.0f;
+		return EVoiceError::VOICE_NOT_PLAYING;
+	}
+
+	EVoiceError::Type FVoice::GetBaselinePitchScale(float& OutBaselinePitchScale) const
+	{
+		AUDIO_VOICE_CHECK_ERROR();
+
+		OutBaselinePitchScale = Params.BaselinePitchScale;
+		return EVoiceError::NONE;
+	}
+
+	EVoiceError::Type FVoice::GetPitchProduct(float& OutPitchProduct) const
+	{
+		AUDIO_VOICE_CHECK_ERROR();
+
+		if (IsPlaying())
+		{
+			DEBUG_AUDIO_CHECK(VoiceDataIndex != INDEX_NONE);
+			FVoiceManager& VoiceManager = AudioModule->GetVoiceManager();
+			return VoiceManager.GetPitchProduct(VoiceDataIndex, OutPitchProduct);
+		}
+		OutPitchProduct = 1.0f;
+		return EVoiceError::VOICE_NOT_PLAYING;
+	}
+
+	EVoiceError::Type FVoice::GetVolumeScale(float& OutVolumeScale) const
+	{
+		AUDIO_VOICE_CHECK_ERROR();
+
+		if (IsPlaying())
+		{
+			DEBUG_AUDIO_CHECK(VoiceDataIndex != INDEX_NONE);
+			FVoiceManager& VoiceManager = AudioModule->GetVoiceManager();
+			return VoiceManager.GetVolumeScale(VoiceDataIndex, OutVolumeScale);
+		}
+
+		OutVolumeScale = 1.0f;
+		return EVoiceError::VOICE_NOT_PLAYING;
+	}
+
+	EVoiceError::Type FVoice::GetBaselineVolumeScale(float& OutBaselineVolumeScale) const
+	{
+		AUDIO_VOICE_CHECK_ERROR();
+
+		OutBaselineVolumeScale = Params.BaselineVolumeScale;
+		return EVoiceError::NONE;
+	}
+
+	EVoiceError::Type FVoice::GetVolumeAttenuation(float& OutAttenuation) const
+	{
+		AUDIO_VOICE_CHECK_ERROR();
+
+		if (IsPlaying())
+		{
+			DEBUG_AUDIO_CHECK(VoiceDataIndex != INDEX_NONE);
+			FVoiceManager& VoiceManager = AudioModule->GetVoiceManager();
+			return VoiceManager.GetVolumeAttenuation(VoiceDataIndex, OutAttenuation);
+		}
+
+		OutAttenuation = 1.0f;
+		return EVoiceError::VOICE_NOT_PLAYING;
+	}
+
+	EVoiceError::Type FVoice::GetVolumeProduct(float& OutVolumeProduct) const
+	{
+		AUDIO_VOICE_CHECK_ERROR();
+
+		if (IsPlaying())
+		{
+			DEBUG_AUDIO_CHECK(VoiceDataIndex != INDEX_NONE);
+			FVoiceManager& VoiceManager = AudioModule->GetVoiceManager();
+			return VoiceManager.GetVolumeProduct(VoiceDataIndex, OutVolumeProduct);
+		}
+		OutVolumeProduct = 1.0f;
+		return EVoiceError::VOICE_NOT_PLAYING;
+	}
+
+	EVoiceError::Type FVoice::GetDuration(float& OutDurationSec)
+	{
+		AUDIO_VOICE_CHECK_ERROR();
+
+		OutDurationSec = DurationSeconds;
+
+		return EVoiceError::NONE;
+	}
+
+	EVoiceError::Type FVoice::GetVoiceState(EVoiceState::Type& OutVoiceState) const
+	{
+		AUDIO_VOICE_CHECK_ERROR();
+
+		OutVoiceState = VoiceState;
+
+		return EVoiceError::NONE;
+	}
+
+	EVoiceError::Type FVoice::GetPlayingState(EVoicePlayingState::Type& OutPlayingState) const
+	{
+		AUDIO_VOICE_CHECK_ERROR();
+
+		OutPlayingState = PlayingState;
+
+		return EVoiceError::NONE;
+	}
+
+	EVoiceError::Type FVoice::GetId(uint32& OutVoiceId) const
+	{
+		AUDIO_VOICE_CHECK_ERROR();
+
+		OutVoiceId = VoiceHandle.Id;
+
+		return EVoiceError::NONE;
+	}
+
+	EVoiceError::Type FVoice::IsValid(bool& bOutIsValid) const
+	{
+		AUDIO_VOICE_CHECK_ERROR();
+
+		bOutIsValid = AudioModule->GetVoiceManager().IsValidVoice(VoiceHandle);
+
+		return EVoiceError::NONE;
+	}
+
+	EVoiceError::Type FVoice::IsPlaying(bool& bOutIsPlaying) const
+	{
+		AUDIO_VOICE_CHECK_ERROR();
+
+		bOutIsPlaying = (PlayingState == EVoicePlayingState::PLAYING_REAL || PlayingState == EVoicePlayingState::PLAYING_VIRTUAL);
+
+		return EVoiceError::NONE;
+	}
+
+	void FVoice::Init(TSharedPtr<FVoice> InThisObject)
+	{
+		ThisObject = InThisObject;
+	}
+
+	void FVoice::NotifyPlayReal(uint32 InVoiceDataIndex)
+	{
+		DEBUG_AUDIO_CHECK(InVoiceDataIndex != INDEX_NONE);
+		VoiceDataIndex = InVoiceDataIndex;
+		PlayingState = EVoicePlayingState::PLAYING_REAL;
+
+		for (int32 i = 0; i < VoiceListeners.Num(); ++i)
+		{
+			VoiceListeners[i]->OnVoiceReal(this);
+		}
+	}
+
+	void FVoice::NotifyPlayVirtual(uint32 InVoiceDataIndex)
+	{
+		DEBUG_AUDIO_CHECK(InVoiceDataIndex != INDEX_NONE);
+		VoiceDataIndex = InVoiceDataIndex;
+		PlayingState = EVoicePlayingState::PLAYING_VIRTUAL;
+
+		for (int32 i = 0; i < VoiceListeners.Num(); ++i)
+		{
+			VoiceListeners[i]->OnVoiceVirtual(this);
+		}
+	}
+
+	void FVoice::NotifyDone()
+	{
+		VoiceDataIndex = INDEX_NONE;
+		PlayingState = EVoicePlayingState::NOT_PLAYING;
+
+		for (int32 i = 0; i < VoiceListeners.Num(); ++i)
+		{
+			VoiceListeners[i]->OnVoiceDone(this);
+		}
+	}
+
+	void FVoice::NotifySuspend()
+	{
+		VoiceDataIndex = INDEX_NONE;
+		PlayingState = EVoicePlayingState::NOT_PLAYING;
+
+		for (int32 i = 0; i < VoiceListeners.Num(); ++i)
+		{
+			VoiceListeners[i]->OnVoiceSuspend(this);
+		}
+	}
+
+	void FVoice::GetHandle(FVoiceHandle& OutVoiceHandle) const
+	{
+		OutVoiceHandle = VoiceHandle;
+	}
+
+	EVoiceError::Type FVoice::SetError(EVoiceError::Type Error)
+	{
+		if (Error != EVoiceError::NONE)
+		{
+			LastError = Error;
+			VoiceState = EVoiceState::HAS_ERROR;
+		}
+		return Error;
+	}
+
+	bool FVoice::IsPlaying() const
+	{
+		return PlayingState == EVoicePlayingState::PLAYING_REAL || PlayingState == EVoicePlayingState::PLAYING_VIRTUAL;
+	}
+
+}
+
+#endif // #if ENABLE_UNREAL_AUDIO
+
diff --git a/Engine/Source/Runtime/UnrealAudio/Private/UnrealAudioVoiceInternal.h b/Engine/Source/Runtime/UnrealAudio/Private/UnrealAudioVoiceInternal.h
new file mode 100644
--- /dev/null
+++ b/Engine/Source/Runtime/UnrealAudio/Private/UnrealAudioVoiceInternal.h
@@ -0,0 +1,120 @@
+#pragma once 
+
+#include "UnrealAudioModule.h"
+#include "UnrealAudioPrivate.h"
+#include "UnrealAudioTypes.h"
+#include "UnrealAudioVoice.h"
+#include "UnrealAudioPrivate.h"
+#include "UnrealAudioEntityManager.h"
+#include "UnrealAudioEmitterInternal.h"
+#include "UnrealAudioUtilities.h"
+
+#if ENABLE_UNREAL_AUDIO
+
+namespace UAudio
+{
+	/**
+	* EVoiceFlag
+	* Flags used in a voice.
+	*/
+	namespace EVoiceFlag
+	{
+		enum Type
+		{
+			NONE		= 0,
+			LOOPING		= (1 << 0),
+			SPATIALIZED = (1 << 1),
+		};
+	}
+
+	/** 
+	* FVoice
+	* Internal concrete implementation of IVoice
+	*/
+	class FVoice : public IVoice
+	{
+	public:
+		FVoice(class FUnrealAudioModule* InParentModule, const FVoiceInitializationParams& InParams);
+		~FVoice() override;
+
+		/************************************************************************/
+		/* IVoice																*/
+		/************************************************************************/
+
+		EVoiceError::Type Play() override;
+		EVoiceError::Type Pause(float InFadeTimeSec = 0.0f) override;
+		EVoiceError::Type Stop(float InFadeTimeSec = 0.0f) override;
+
+		EVoiceError::Type AddVoiceListener(IVoiceListener* Listener) override;
+		EVoiceError::Type RemoveVoiceListener(IVoiceListener* Listener) override;
+
+		EVoiceError::Type SetVolumeScale(float InVolumeLinear, float InFadeTimeSec = 0.0f) override;
+		EVoiceError::Type SetPitchScale(float InPitchScale, float InFadeTimeSec = 0.0f) override;
+
+		EVoiceError::Type GetPitchScale(float& OutPitchScale) const override;
+		EVoiceError::Type GetBaselinePitchScale(float& OutBaselinePitchScale) const override;
+		EVoiceError::Type GetPitchProduct(float& OutPitchProduct) const override;
+		EVoiceError::Type GetVolumeScale(float& OutVolumeScale) const override;
+		EVoiceError::Type GetBaselineVolumeScale(float& OutBaselineVolumeScale) const override;
+		EVoiceError::Type GetVolumeAttenuation(float& OutAttenuation) const override;
+		EVoiceError::Type GetVolumeProduct(float& OutVolumeProduct) const override;
+		EVoiceError::Type GetDuration(float& OutDurationSec) override;
+		EVoiceError::Type GetVoiceState(EVoiceState::Type& OutVoicestate) const override;
+		EVoiceError::Type GetPlayingState(EVoicePlayingState::Type& OutPlayingState) const override;
+		EVoiceError::Type GetId(uint32& OutVoiceId) const override;
+
+		EVoiceError::Type IsValid(bool& bOutIsValid) const override;
+		EVoiceError::Type IsPlaying(bool& OutIsplaying) const override;
+
+		/************************************************************************/
+		/* Public internal API													*/
+		/************************************************************************/
+
+		void Init(TSharedPtr<FVoice> InThisObject);
+		void NotifyPlayReal(uint32 InVoiceDataIndex);
+		void NotifyPlayVirtual(uint32 InVoiceDataIndex);
+		void NotifyDone();
+		void NotifySuspend();
+		void GetHandle(FVoiceHandle& OutHandle) const;
+
+	private:
+		EVoiceError::Type SetError(EVoiceError::Type Error);
+		bool IsPlaying() const;
+
+		struct FPlayVoiceInfo
+		{
+			float VolumeScale;
+			float VolumeScaleTime;
+			float PitchScale;
+			float PitchScaletime;
+
+			FPlayVoiceInfo()
+				: VolumeScale(1.0f)
+				, VolumeScaleTime(0.0f)
+				, PitchScale(1.0f)
+				, PitchScaletime(0.0f)
+			{}
+		};
+
+		FUnrealAudioModule*			AudioModule;
+		EVoicePlayingState::Type	PlayingState;
+		EVoiceState::Type			VoiceState;
+		FVoiceInitializationParams	Params;
+		FPlayVoiceInfo				PlayInfo;
+		TArray<IVoiceListener*>		VoiceListeners;
+		FVoiceHandle				VoiceHandle;
+
+		// index into the real/virtual voice data array
+		uint32						VoiceDataIndex;
+
+		TSharedPtr<FVoice>			ThisObject;
+
+		EVoiceError::Type			LastError;
+		float						DurationSeconds;
+		bool						bPlayCalled;
+	};
+
+
+}
+
+#endif // #if ENABLE_UNREAL_AUDIO
\ No newline at end of file
diff --git a/Engine/Source/Runtime/UnrealAudio/Private/UnrealAudioVoiceManager.cpp b/Engine/Source/Runtime/UnrealAudio/Private/UnrealAudioVoiceManager.cpp
new file mode 100644
--- /dev/null
+++ b/Engine/Source/Runtime/UnrealAudio/Private/UnrealAudioVoiceManager.cpp
@@ -0,0 +1,652 @@
+// Copyright 1998-2015 Epic Games, Inc. All Rights Reserved.
+
+#include "UnrealAudioPrivate.h"
+#include "UnrealAudioModule.h"
+#include "UnrealAudioVoiceInternal.h"
+#include "UnrealAudioVoiceManager.h"
+
+#define UNREAL_AUDIO_VOICE_COMMAND_QUEUE_SIZE (50)
+
+#define VOICE_CHECK_INITIALIZATION()													\
+	if (VoiceState == EVoiceState::UNINITIALIZED)										\
+	{																					\
+		UA_SYSTEM_ERROR(ESystemError::NOT_INITIALIZED, TEXT("Voice not initialized.")); \
+		return ESystemError::NOT_INITIALIZED;											\
+	}
+
+
+#if ENABLE_UNREAL_AUDIO
+
+namespace UAudio
+{
+	/************************************************************************/
+	/* Main Thread Functions												*/
+	/************************************************************************/
+
+	FVoiceManager::FVoiceManager(FUnrealAudioModule* InAudioModule)
+		: EntityManager(500)
+		, AudioModule(InAudioModule)
+		, PitchManager(InAudioModule)
+		, VolumeManager(InAudioModule)
+		, VoiceMixer(InAudioModule, this)
+		, NumVirtualVoices(0)
+		, NumRealVoices(0)
+	{
+		Settings = { 0 };
+	}
+
+	FVoiceManager::~FVoiceManager()
+	{
+
+	}
+
+	void FVoiceManager::Init(const FVoiceManagerSettings& InSettings)
+	{
+		DEBUG_AUDIO_CHECK(InSettings.MaxVoiceCount > 0);
+		DEBUG_AUDIO_CHECK(InSettings.MaxVirtualVoiceCount > 0);
+		DEBUG_AUDIO_CHECK(Settings.MaxVoiceCount == 0);
+
+		Settings = InSettings;
+		uint32 TotalNumVoices = Settings.MaxVirtualVoiceCount + Settings.MaxVoiceCount;
+
+		VoiceHandles.Init(FVoiceHandle(), TotalNumVoices);
+		States.Init(EVoiceState::UNINITIALIZED, TotalNumVoices);
+		PlayingStates.Init(EVoicePlayingState::NOT_PLAYING, TotalNumVoices);
+		Flags.Init(0, TotalNumVoices);
+		VolumeManager.Init(TotalNumVoices);
+		PitchManager.Init(TotalNumVoices);
+		VoiceMixer.Init(InSettings);
+
+		for (uint32 i = 0; i < TotalNumVoices; ++i)
+		{
+			FreeDataIndicies.Enqueue(i);
+		}
+	}
+
+	EVoiceError::Type FVoiceManager::PlayVoice(FVoice* Voice, const FVoiceInitializationData* InitData, FVoiceHandle& OutVoiceHandle)
+	{
+		// First create the voice handle
+		FVoiceHandle VoiceHandle(EntityManager.CreateEntity());
+
+		// Store the parent voice object ptr at the index associated with this handle.
+		// This will allow us to send messages back to that voice object from the audio thread
+
+		uint32 VoiceIndex = VoiceHandle.GetIndex();
+		if (VoiceIndex < (uint32)VoiceObjects.Num())
+		{
+			check(VoiceObjects[VoiceIndex] == nullptr);
+			VoiceObjects[VoiceIndex] = Voice;
+		}
+		else
+		{
+			check(VoiceIndex == (uint32)VoiceObjects.Num());
+			VoiceObjects.Add(Voice);
+		}
+
+		OutVoiceHandle = VoiceHandle;
+
+		AudioModule->SendAudioThreadCommand(FCommand(EAudioThreadCommand::VOICE_PLAY, VoiceHandle, (void*)InitData));
+
+		return EVoiceError::NONE;
+	}
+
+	EVoiceError::Type FVoiceManager::PauseVoice(const FVoiceHandle& VoiceHandle, float InFadeTimeSec)
+	{
+		AudioModule->SendAudioThreadCommand(FCommand(EAudioThreadCommand::VOICE_PAUSE, VoiceHandle, InFadeTimeSec));
+		return EVoiceError::NONE;
+	}
+
+	EVoiceError::Type FVoiceManager::StopVoice(const FVoiceHandle& VoiceHandle, float InFadeTimeSec)
+	{
+		AudioModule->SendAudioThreadCommand(FCommand(EAudioThreadCommand::VOICE_STOP, VoiceHandle, InFadeTimeSec));
+		return EVoiceError::NONE;
+	}
+
+	EVoiceError::Type FVoiceManager::SetVolumeScale(uint32 VoiceDataIndex, float InVolumeScale, float InFadeTimeSec)
+	{
+		AudioModule->SendAudioThreadCommand(FCommand(EAudioThreadCommand::VOICE_SET_VOLUME_SCALE, VoiceDataIndex, InVolumeScale, InFadeTimeSec));
+		return EVoiceError::NONE;
+	}
+
+	EVoiceError::Type FVoiceManager::SetPitchScale(uint32 VoiceDataIndex, float InPitchScale, float InFadeTimeSec)
+	{
+		AudioModule->SendAudioThreadCommand(FCommand(EAudioThreadCommand::VOICE_SET_PITCH_SCALE, VoiceDataIndex, InPitchScale, InFadeTimeSec));
+		return EVoiceError::NONE;
+	}
+
+	void FVoiceManager::NotifyVoiceDone(const FCommand& Command)
+	{
+		check(Command.NumArguments == 1);
+		check(Command.Arguments[0].DataType == ECommandData::HANDLE);
+
+		FVoiceHandle VoiceHandle = Command.Arguments[0].Data.Handle;
+		uint32 Index = VoiceHandle.GetIndex();
+
+		// Get the voice object
+		FVoice* Voice = VoiceObjects[Index];
+		check(Voice);
+
+		// Notify that its done
+		Voice->NotifyDone();
+
+		// Free the slot in the voice objects array
+		VoiceObjects[Index] = nullptr;
+
+		// And free the handle
+		EntityManager.ReleaseEntity(VoiceHandle);
+	}
+
+	void FVoiceManager::NotifyVoiceReal(const FCommand& Command)
+	{
+		check(Command.NumArguments == 2);
+		check(Command.Arguments[0].DataType == ECommandData::HANDLE);
+		check(Command.Arguments[1].DataType == ECommandData::UINT32);
+
+		FVoiceHandle VoiceHandle = Command.Arguments[0].Data.Handle;
+		uint32 Index = VoiceHandle.GetIndex();
+
+		uint32 VoiceDataIndex = Command.Arguments[1].Data.UnsignedInt32;
+
+		// Get the voice object
+		FVoice* Voice = VoiceObjects[Index];
+		check(Voice);
+
+		Voice->NotifyPlayReal(VoiceDataIndex);
+	}
+
+	void FVoiceManager::NotifyVoiceVirtual(const FCommand& Command)
+	{
+		check(Command.NumArguments == 2);
+		check(Command.Arguments[0].DataType == ECommandData::HANDLE);
+		check(Command.Arguments[1].DataType == ECommandData::UINT32);
+
+		FVoiceHandle VoiceHandle = Command.Arguments[0].Data.Handle;
+		uint32 Index = VoiceHandle.GetIndex();
+
+		uint32 VoiceDataIndex = Command.Arguments[1].Data.UnsignedInt32;
+
+		// Get the voice object
+		FVoice* Voice = VoiceObjects[Index];
+		check(Voice != nullptr);
+
+		Voice->NotifyPlayVirtual(VoiceDataIndex);
+	}
+
+	void FVoiceManager::NotifyVoiceSuspend(const FCommand& Command)
+	{
+		check(Command.NumArguments == 1);
+		check(Command.Arguments[0].DataType == ECommandData::HANDLE);
+
+		FVoiceHandle VoiceHandle = Command.Arguments[0].Data.Handle;
+		uint32 Index = VoiceHandle.GetIndex();
+
+		// Get the voice object
+		FVoice* Voice = VoiceObjects[Index];
+		check(Voice);
+
+		Voice->NotifySuspend();
+	}
+
+	EVoiceError::Type FVoiceManager::GetNumPlayingVoices(int32& OutNumPlayingVoices) const
+	{
+		OutNumPlayingVoices = NumRealVoices;
+		return EVoiceError::NONE;
+	}
+
+	EVoiceError::Type FVoiceManager::GetMaxNumberOfPlayingVoices(int32& OutMaxNumPlayingVoices) const
+	{
+		OutMaxNumPlayingVoices = Settings.MaxVoiceCount;
+		return EVoiceError::NONE;
+	}
+
+	EVoiceError::Type FVoiceManager::GetNumVirtualVoices(int32& OutNumVirtualVoices) const
+	{
+		OutNumVirtualVoices = NumVirtualVoices;
+		return EVoiceError::NONE;
+	}
+
+	EVoiceError::Type FVoiceManager::GetVolumeScale(uint32 VoiceDataIndex, float& OutVolumeProduct) const
+	{
+		check(VoiceDataIndex != INDEX_NONE);
+		OutVolumeProduct = VolumeManager.GetVolumeProduct(VoiceDataIndex);
+		return EVoiceError::NONE;
+	}
+
+	EVoiceError::Type FVoiceManager::GetVolumeAttenuation(uint32 VoiceDataIndex, float& OutVolumeAttenuation) const
+	{
+		check(VoiceDataIndex != INDEX_NONE);
+		OutVolumeAttenuation = VolumeManager.GetVolumeAttenuation(VoiceDataIndex);
+		return EVoiceError::NONE;
+	}
+
+	EVoiceError::Type FVoiceManager::GetVolumeFade(uint32 VoiceDataIndex, float& OutVolumeFade) const
+	{
+		check(VoiceDataIndex != INDEX_NONE);
+		OutVolumeFade = VolumeManager.GetVolumeFade(VoiceDataIndex);
+		return EVoiceError::NONE;
+	}
+
+	EVoiceError::Type FVoiceManager::GetVolumeProduct(uint32 VoiceDataIndex, float& OutVolumeProduct) const
+	{
+		check(VoiceDataIndex != INDEX_NONE);
+		OutVolumeProduct = VolumeManager.GetVolumeProduct(VoiceDataIndex);
+		return EVoiceError::NONE;
+	}
+
+	EVoiceError::Type FVoiceManager::GetPitchScale(uint32 VoiceDataIndex, float& OutPitchScale) const
+	{
+		check(VoiceDataIndex != INDEX_NONE);
+		OutPitchScale = PitchManager.GetPitchScale(VoiceDataIndex);
+		return EVoiceError::NONE;
+	}
+
+	EVoiceError::Type FVoiceManager::GetPitchProduct(uint32 VoiceDataIndex, float& OutPitchProduct) const
+	{
+		check(VoiceDataIndex != INDEX_NONE);
+		OutPitchProduct = PitchManager.GetPitchProduct(VoiceDataIndex);
+		return EVoiceError::NONE;
+	}
+
+	bool FVoiceManager::IsValidVoice(const FVoiceHandle& VoiceHandle) const
+	{
+		return EntityManager.IsValidEntity(VoiceHandle);
+	}
+
+	/************************************************************************/
+	/* Audio Thread Functions												*/
+	/************************************************************************/
+
+	void FVoiceManager::PlayVoice(const FCommand& Command)
+	{
+		check(Command.NumArguments == 2);
+		check(Command.Arguments[0].DataType == ECommandData::HANDLE);
+		check(Command.Arguments[1].DataType == ECommandData::POINTER);
+
+		FVoiceHandle VoiceHandle = Command.Arguments[0].Data.Handle;
+		uint32 VoiceIndex = VoiceHandle.GetIndex();
+
+		check(Command.Arguments[1].Data.PtrVal != nullptr);
+		FVoiceInitializationData* VoiceInitData = (FVoiceInitializationData*)Command.Arguments[1].Data.PtrVal;
+
+		EVoicePlayingState::Type PlayType = EVoicePlayingState::SUSPENDED;
+		int32 SuspendedDataVoiceIndex = INDEX_NONE;
+		bool bPlayReal = false;
+		float VolumeProduct = -1.0f;
+		float VolumeAttenuation = -1.0f;
+
+		// If there are any free real voice slots, then this is easy: we play it as a real voice
+		if (NumRealVoices < Settings.MaxVoiceCount)
+		{
+			PlayType = EVoicePlayingState::PLAYING_REAL;
+		}
+		// ... otherwise we need to figure out what to do with the voice
+		else
+		{
+			// Get the weighted priority of the new voice
+			float WeightedPriority = 0.0f;
+			ComputeInitialVolumeValues(VoiceInitData, VolumeProduct, VolumeAttenuation, WeightedPriority);
+			check(VolumeProduct >= 0.0f && VolumeProduct <= 1.0f);
+
+			// if any currently playing real voices have a lower priority than this voice, they need to become virtual (or suspended)
+			const TArray<FSortedVoiceEntry>& SortedEntries = VolumeManager.GetSortedVoices();
+			check(SortedEntries.Num() == (Settings.MaxVoiceCount + Settings.MaxVirtualVoiceCount));
+
+			// Get the least-priority real voice entry 
+			const FSortedVoiceEntry& LeastPriorityRealEntry = SortedEntries[Settings.MaxVoiceCount - 1];
+
+			// If the new voice's weighted priority is higher than the least real voice...
+			if (WeightedPriority > LeastPriorityRealEntry.PriorityWeightedVolume)
+			{
+				// ... then we're going to play the new voice as a real voice 
+				PlayType = EVoicePlayingState::PLAYING_REAL;
+
+				// Now we figure out  what to do with the real voice we're going to "steal"...
+
+				// First, sanity check and make sure the least-pri playing voice playing at this index is indeed real
+				check(PlayingStates[LeastPriorityRealEntry.Index] == EVoicePlayingState::PLAYING_REAL);
+
+				// if we have any virtual voices at all...
+				if (Settings.MaxVirtualVoiceCount > 0)
+				{
+					// ... and if we don't have any virtual voice slots left...
+					if (NumVirtualVoices == Settings.MaxVirtualVoiceCount)
+					{
+						// ... then suspend the least priority virtual voice
+						const FSortedVoiceEntry& LastVirtualEntry = SortedEntries[SortedEntries.Num() - 1];
+						SuspendedDataVoiceIndex = LastVirtualEntry.Index;
+					}
+
+					// Now make the playing real voice virtual
+					// TODO: un-register this voice as a real voice, and turn it virtual
+					PlayingStates[LeastPriorityRealEntry.Index] = EVoicePlayingState::PLAYING_VIRTUAL;
+				}
+				else
+				{
+					// ... otherwise just suspend the lower-priority real voice
+					SuspendedDataVoiceIndex = LeastPriorityRealEntry.Index;
+				}
+			}
+			// ... else if we have any virtual voice slots at all...
+			else if (Settings.MaxVirtualVoiceCount > 0)
+			{
+				// ... but if we don't have any free virtual voices left
+				if (NumVirtualVoices == Settings.MaxVirtualVoiceCount)
+				{
+					// ... check the least priority virtual voice entry (i.e. the last entry)
+					const FSortedVoiceEntry& LeastPriorityVirtualEntry = SortedEntries[SortedEntries.Num() - 1];
+
+					// if the new voice has a higher priority than the last entry
+					if (WeightedPriority > LeastPriorityVirtualEntry.PriorityWeightedVolume)
+					{
+						// ... then we're going to play that new voice as a virtual voice
+						PlayType = EVoicePlayingState::PLAYING_VIRTUAL;
+
+						// and suspend the last virtual voice entry
+						SuspendedDataVoiceIndex = LeastPriorityRealEntry.Index;
+					}
+					// ... otherwise the new voice is lower priority than even the lowest-priority virtual voice
+					else
+					{
+						// and needs to be suspended
+						PlayType = EVoicePlayingState::SUSPENDED;
+					}
+				}
+				// .. otherwise we must have room in our virtual voice allocation
+				else
+				{
+					// But sanity check that we indeed have room left
+					check(NumVirtualVoices < Settings.MaxVirtualVoiceCount);
+
+					// Set the new voice to play as virtual
+					PlayType = EVoicePlayingState::PLAYING_VIRTUAL;
+
+					// No need to suspend anything...
+				}
+			}
+			// ... otherwise, we'll have to immediately suspend the new voice
+			else
+			{
+				PlayType = EVoicePlayingState::SUSPENDED;
+			}
+		}
+
+		// If this new voice is immediately suspended...
+		if (PlayType == EVoicePlayingState::SUSPENDED)
+		{
+			// ... send a message back to main thread of the fact
+			AudioModule->SendMainThreadCommand(FCommand(EMainThreadCommand::VOICE_SUSPEND, VoiceHandle));
+		}
+		else
+		{
+			// We are playing this voice as either real or virtual and need to determine the voice data index
+			uint32 VoiceDataIndex = INDEX_NONE;
+
+			// First, if we had are suspending a virtual or real voice, we need to suspend it and take that index as the new index
+			if (SuspendedDataVoiceIndex != INDEX_NONE)
+			{
+				// Get the voice handle at this index
+				const FVoiceHandle& SuspendedVoiceHandle = VoiceHandles[VoiceIndex];
+
+				// Release the voice index entries in the pitch and volume managers
+				PitchManager.ReleaseEntry(VoiceIndex);
+				VolumeManager.ReleaseEntry(VoiceIndex);
+
+				// Remove this voice handle from the handle-to-index map
+				HandleToIndex.Remove(SuspendedVoiceHandle.GetIndex());
+
+				// ... and set the suspended voice index as the index of the new voice
+				VoiceDataIndex = SuspendedDataVoiceIndex;
+
+				// Send message to main that this voice was suspended
+				AudioModule->SendMainThreadCommand(FCommand(EMainThreadCommand::VOICE_SUSPEND, SuspendedVoiceHandle));
+			}
+			else
+			{
+				// ... otherwise, we need to get a new voice data index. Since we didn't suspend anything,
+				// and this is a real/virtual voice, then we need to get a new voice data index
+				FreeDataIndicies.Dequeue(VoiceDataIndex);
+			}
+
+			// Sanity check that we got a new voice data index
+			check(VoiceDataIndex != INDEX_NONE);
+
+			// Add this VoiceDataIndex to the handle-to-index map
+			HandleToIndex.Add(VoiceHandle.GetIndex(), VoiceDataIndex);
+
+			// We might need to compute the volume product and the volume attenuation for the new voice
+			if (VolumeProduct < 0.0f)
+			{
+				float WeightedPriority = 0.0f;
+				ComputeInitialVolumeValues(VoiceInitData, VolumeProduct, VolumeAttenuation, WeightedPriority);
+			}
+
+			// Now we can just initialize the voice at the suspended index
+			check(VolumeProduct >= 0.0f && VolumeAttenuation >= 0.0f);
+
+			// Initialize data
+			VoiceHandles[VoiceDataIndex] = VoiceHandle;
+			PlayingStates[VoiceDataIndex] = PlayType;
+			States[VoiceDataIndex] = EVoiceState::PLAYING;
+			Flags[VoiceDataIndex] = VoiceInitData->VoiceFlags;
+
+			FPitchInitParam PitchParams =
+			{
+				VoiceInitData->BaselinePitchScale,
+				VoiceInitData->DynamicPitchScale,
+				VoiceInitData->DynamicPitchTime,
+				VoiceInitData->DurationSeconds
+			};
+
+			PitchManager.InitializeEntry(VoiceDataIndex, PitchParams);
+
+			FVolumeInitParam VolumeParams =
+			{
+				VoiceInitData->EmitterHandle,
+				VoiceInitData->BaselineVolumeScale,
+				VoiceInitData->DynamicVolumeScale,
+				VoiceInitData->DynamicVolumeTime,
+				VolumeProduct,
+				VolumeAttenuation,
+				VoiceInitData->PriorityWeight
+			};
+
+			VolumeManager.InitializeEntry(VoiceDataIndex, VolumeParams);
+
+			// Send a message to main thread that this voice is now playing...
+			if (PlayType == EVoicePlayingState::PLAYING_REAL)
+			{
+				AudioModule->SendMainThreadCommand(FCommand(EMainThreadCommand::VOICE_REAL, VoiceHandle, VoiceDataIndex));
+				//AudioModule->SendDeviceThreadCommand(FCommand(EDeviceThreadCommand::VOICE_REAL, VoiceHandle, VoiceDataIndex));
+			}
+			else
+			{
+				AudioModule->SendMainThreadCommand(FCommand(EMainThreadCommand::VOICE_VIRTUAL, VoiceHandle, VoiceDataIndex));
+			}
+		}
+	}
+
+	void FVoiceManager::PauseVoice(const FCommand& Command)
+	{
+		check(Command.NumArguments == 2);
+		check(Command.Arguments[0].DataType == ECommandData::HANDLE);
+		check(Command.Arguments[1].DataType == ECommandData::FLOAT_32);
+
+		FVoiceHandle VoiceHandle = Command.Arguments[0].Data.Handle;
+		float FadeTimeSec = Command.Arguments[1].Data.Float32Val;
+
+		uint32 VoiceDataIndex = GetVoiceDataIndex(VoiceHandle);
+
+		States[VoiceDataIndex] = EVoiceState::PAUSING;
+		VolumeManager.SetFadeOut(VoiceDataIndex, FadeTimeSec);
+	}
+
+	void FVoiceManager::StopVoice(const FCommand& Command)
+	{
+		check(Command.NumArguments == 2);
+		check(Command.Arguments[0].DataType == ECommandData::HANDLE);
+		check(Command.Arguments[1].DataType == ECommandData::FLOAT_32);
+
+		FVoiceHandle VoiceHandle = Command.Arguments[0].Data.Handle;
+		float FadeTimeSec = Command.Arguments[1].Data.Float32Val;
+
+		uint32 VoiceDataIndex = GetVoiceDataIndex(VoiceHandle);
+		States[VoiceDataIndex] = EVoiceState::STOPPING;
+		VolumeManager.SetFadeOut(VoiceDataIndex, FadeTimeSec);
+	}
+
+	void FVoiceManager::SetVolumeScale(const FCommand& Command)
+	{
+		check(Command.NumArguments == 3);
+		check(Command.Arguments[0].DataType == ECommandData::UINT32);
+		check(Command.Arguments[1].DataType == ECommandData::FLOAT_32);
+		check(Command.Arguments[2].DataType == ECommandData::FLOAT_32);
+
+		uint32 VoiceDataIndex = Command.Arguments[0].Data.UnsignedInt32;
+		float VolumeScale = Command.Arguments[1].Data.Float32Val;
+		float DeltaTimeSeconds = Command.Arguments[2].Data.Float32Val;
+
+		VolumeManager.SetDynamicVolumeScale(VoiceDataIndex, VolumeScale, DeltaTimeSeconds);
+	}
+
+	void FVoiceManager::SetPitchScale(const FCommand& Command)
+	{
+		check(Command.NumArguments == 3);
+		check(Command.Arguments[0].DataType == ECommandData::UINT32);
+		check(Command.Arguments[1].DataType == ECommandData::FLOAT_32);
+		check(Command.Arguments[2].DataType == ECommandData::FLOAT_32);
+
+		uint32 VoiceDataIndex = Command.Arguments[0].Data.UnsignedInt32;
+		float PitchScale = Command.Arguments[1].Data.Float32Val;
+		float DeltaTimeSeconds = Command.Arguments[2].Data.Float32Val;
+
+		PitchManager.SetDynamicPitchScale(VoiceDataIndex, PitchScale, DeltaTimeSeconds);
+	}
+
+	void FVoiceManager::Update()
+	{
+		VolumeManager.Update();
+		PitchManager.Update();
+
+		// Handle any changes virtual/real voice changes after the update
+		UpdateStates();
+	}
+
+	void FVoiceManager::UpdateStates()
+	{
+		//////////////////////////////////////////////////////////////////////////
+		// Update virtual and real voices, swapping each other based on a sorted volume array
+		//
+		// 1) Skip not-playing voices since they aren't playing at all (i.e. they're paused or stopped)
+		//
+		// 2) Keep two separate counts, one count is the actual new real voice count after some voices
+		//		have possibly turned from virtual to real. The other count is the old real-voice count.
+		//		We use the previous playing sound count to early-out the iteration.
+		//		The other counter is used to track the new real voice count. Once a voice which was
+		//		previously real is found while the number of new voice counts are higher than
+		//		the max voice count, that voice (and subsequent voices) will be virtualized.
+		//
+		//////////////////////////////////////////////////////////////////////////
+
+		const TArray<FSortedVoiceEntry>& SortedVoiceEntries = VolumeManager.GetSortedVoices();
+
+		// If there are no voices playing then there's nothing to do...
+		if (SortedVoiceEntries.Num() == 0)
+		{
+			return;
+		}
+
+		// Keep track of how many real-playing voices were playing in the last update
+		int32 PrevRealCount = 0;
+
+		// Keep track of how many real-playing voices are going to be playing in the next update
+		int32 NextRealCount = 0;
+
+		// Loop through the highest-pri voice entries (from the front)
+		for (const FSortedVoiceEntry& Entry : SortedVoiceEntries)
+		{
+			// Get the state of the voice
+			EVoicePlayingState::Type PlayingState = PlayingStates[Entry.Index];
+			check(PlayingState != EVoicePlayingState::NOT_PLAYING);
+
+			// If it's a real voice...
+			if (PlayingState == EVoicePlayingState::PLAYING_REAL)
+			{
+				// count this voice as a "next" real-voice
+				++NextRealCount;
+
+				// if the number of real voices for the next update is already greater than
+				// the max number of voices, then we must suspend this voice entry
+				if (NextRealCount >= Settings.MaxVoiceCount)
+				{
+					PlayingStates[Entry.Index] = EVoicePlayingState::PLAYING_VIRTUAL;
+
+					// Notify the main thread that this voice has become virtual
+					AudioModule->SendMainThreadCommand(FCommand(EMainThreadCommand::VOICE_VIRTUAL, VoiceHandles[Entry.Index]));
+
+					// TODO: Remove this DataVoiceIndex from active playing sounds
+				}
+
+				if (++PrevRealCount == Settings.MaxVoiceCount)
+				{
+					break;
+				}
+			}
+			// else if this is a virtual voice, then this must be a higher pri voice than one inside the highest voice pri's
+			else if (PlayingState == EVoicePlayingState::PLAYING_VIRTUAL)
+			{
+				check(NextRealCount < Settings.MaxVoiceCount);
+
+				// count this new voice as a real voice for the next update loop
+				++NextRealCount;
+
+				// Turn the virtual voice into a real voice
+				PlayingStates[Entry.Index] = EVoicePlayingState::PLAYING_REAL;
+
+				// Notify the main thread that this voice has become real
+				AudioModule->SendMainThreadCommand(FCommand(EMainThreadCommand::VOICE_REAL, VoiceHandles[Entry.Index]));
+
+				// TODO: ADD THIS STATE TO ACTIVE PLAYING SOUNDS AND SOUND OUTPUT
+			}
+		}
+	}
+
+	void FVoiceManager::ComputeInitialVolumeValues(const FVoiceInitializationData* VoiceInitData, float& OutVolumeProduct, float& OutVolumeAttenuation, float& OutWeightedPriority) const
+	{
+		// Initialize the volume product to the baseline volume scale
+		float VolumeProduct = VoiceInitData->BaselineVolumeScale;
+
+		// If the new volume-time for the volume-scale is 0.0, then we immediately factor the volume scale
+		if (VoiceInitData->DynamicVolumeTime == 0.0f)
+		{
+			VolumeProduct *= VoiceInitData->DynamicVolumeScale;
+		}
+
+		// If this is a 3d voice, then we need to factor in attenuation
+		OutVolumeAttenuation = 1.0f;
+		if (VoiceInitData->VoiceFlags & EVoiceFlag::SPATIALIZED)
+		{
+			// TODO: Plug this into attenuation manager/code to get proper attenuation
+			VolumeProduct *= OutVolumeAttenuation;
+		}
+
+		OutVolumeProduct = VolumeProduct;
+		OutWeightedPriority = VoiceInitData->PriorityWeight * VolumeProduct;
+	}
+
+	uint32 FVoiceManager::GetVoiceDataIndex(const FVoiceHandle& VoiceHandle) const
+	{
+		uint32 VoiceIndex = VoiceHandle.GetIndex();
+		const uint32 * VoiceDataIndex = HandleToIndex.Find(VoiceIndex);
+		check(VoiceDataIndex != nullptr);
+		return *VoiceDataIndex;
+	}
+
+	void FVoiceManager::Mix(struct FCallbackInfo& CallbackInfo)
+	{
+
+	}
+
+}
+
+#endif // #if ENABLE_UNREAL_AUDIO
+
diff --git a/Engine/Source/Runtime/UnrealAudio/Private/UnrealAudioVoiceManager.h b/Engine/Source/Runtime/UnrealAudio/Private/UnrealAudioVoiceManager.h
new file mode 100644
--- /dev/null
+++ b/Engine/Source/Runtime/UnrealAudio/Private/UnrealAudioVoiceManager.h
@@ -0,0 +1,164 @@
+#pragma once 
+
+#include "UnrealAudioPrivate.h"
+#include "UnrealAudioModule.h"
+#include "UnrealAudioTypes.h"
+#include "UnrealAudioVoice.h"
+#include "UnrealAudioEntityManager.h"
+#include "UnrealAudioEmitterInternal.h"
+#include "UnrealAudioVoiceInternal.h"
+#include "UnrealAudioUtilities.h"
+#include "UnrealAudioPitchManager.h"
+#include "UnrealAudioVolumeManager.h"
+#include "UnrealAudioVoiceMixer.h"
+
+#if ENABLE_UNREAL_AUDIO
+
+namespace UAudio
+{
+	/**
+	* FVoiceSettings
+	* Voice settings struct
+	*/
+	struct FVoiceManagerSettings
+	{
+		/** The min pitch of a voice. */
+		float MinPitch;
+
+		/** The max pitch of a voice. */
+		float MaxPitch;
+
+		/** The max voice count of the voice manager. */
+		int32 MaxVoiceCount;
+
+		/** The max virtual voice count of the voice manager. */
+		int32 MaxVirtualVoiceCount;
+
+		/** The number of sound file decoders to use to decode audio data for voice playback. */
+		uint32 NumDecoders;
+
+		/** Settings structure for sound file decoding. */
+		FSoundFileDecoderSettings DecoderSettings;
+
+		/** How fast the control data is updating (i.e. update speed of audio system thread). */
+		float ControlUpdateRateSeconds;
+	};
+
+	/**
+	* FVoiceInitializationData
+	* Struct used to initialize data for a single voice entry from main thread to the audio thread.
+	*/
+	struct FVoiceInitializationData
+	{
+		TSharedPtr<ISoundFile> SoundFile;
+		FEmitterHandle EmitterHandle;
+		float BaselineVolumeScale;
+		float DynamicVolumeScale;
+		float DynamicVolumeTime;
+		float BaselinePitchScale;
+		float DynamicPitchScale;
+		float DynamicPitchTime;
+		float PriorityWeight;
+		float DurationSeconds;
+		uint32 VoiceFlags;
+	};
+
+	/**
+	* FVoiceManager
+	* Manages voice handles, messaging, and updating.
+	*/
+	class FVoiceManager
+	{
+	public:
+
+		/************************************************************************/
+		/* Main Thread Functions												*/
+		/************************************************************************/
+
+		FVoiceManager(FUnrealAudioModule* InAudioModule);
+		~FVoiceManager();
+
+		void Init(const FVoiceManagerSettings& Settings);
+
+		// Commands to audio thread...
+		EVoiceError::Type PlayVoice(FVoice* Voice, const FVoiceInitializationData* InitData, FVoiceHandle& OutVoiceHandle);
+		EVoiceError::Type PauseVoice(const FVoiceHandle& VoiceHandle, float InFadeTimeSec);
+		EVoiceError::Type StopVoice(const FVoiceHandle& VoiceHandle, float InFadeTimeSec);
+		EVoiceError::Type SetVolumeScale(uint32 VoiceDataIndex, float InVolumeScale, float InFadeTimeSec);
+		EVoiceError::Type SetPitchScale(uint32 VoiceDataIndex, float InPitchScale, float InFadeTimeSec);
+
+		// Commands from audio thread
+		void NotifyVoiceDone(const FCommand& Command);
+		void NotifyVoiceReal(const FCommand& Command);
+		void NotifyVoiceVirtual(const FCommand& Command);
+		void NotifyVoiceSuspend(const FCommand& Command);
+
+		// Getters
+		EVoiceError::Type GetNumPlayingVoices(int32& OutNumPlayingVoices) const;
+		EVoiceError::Type GetMaxNumberOfPlayingVoices(int32& OutMaxNumPlayingVoices) const;
+		EVoiceError::Type GetNumVirtualVoices(int32& OutNumVirtualVoices) const;
+		EVoiceError::Type GetMaxNumberOfVirtualVoices(int32& OutMaxNumPlayingVoices) const;
+		EVoiceError::Type GetVolumeScale(uint32 VoiceDataIndex, float& OutVolumeScale) const;
+		EVoiceError::Type GetVolumeAttenuation(uint32 VoiceDataIndex, float& OutVolumeAttenuation) const;
+		EVoiceError::Type GetVolumeFade(uint32 VoiceDataIndex, float& OutVolumeFade) const;
+		EVoiceError::Type GetVolumeProduct(uint32 VoiceDataIndex, float& OutVolumeProduct) const;
+		EVoiceError::Type GetPitchScale(uint32 VoiceDataIndex, float& OutPitchScale) const;
+		EVoiceError::Type GetPitchProduct(uint32 VoiceDataIndex, float& OutPitchProduct) const;
+
+		bool IsValidVoice(const FVoiceHandle& VoiceHandle) const;
+
+		/************************************************************************/
+		/* Audio Thread Functions												*/
+		/************************************************************************/
+
+		void PlayVoice(const FCommand& Command);
+		void PauseVoice(const FCommand& Command);
+		void StopVoice(const FCommand& Command);
+		void SetVolumeScale(const FCommand& Command);
+		void SetPitchScale(const FCommand& Command);
+		void Update();
+
+		/**
+		 * Mixes the active audio-generating voices together, called from audio device thread.
+		 *
+		 * @param [in,out]	CallbackInfo	Information describing the audio device callback.
+		 */
+
+		void Mix(struct FCallbackInfo& CallbackInfo);
+
+	private:
+
+		void UpdateStates();
+		uint32 GetVoiceDataIndex(const FVoiceHandle& VoiceHandle) const;
+		void ComputeInitialVolumeValues(const FVoiceInitializationData* VoiceInitData, float& OutVolumeProduct, float& OutVolumeAttenuation, float& OutWeightedPriority) const;
+
+		// Main thread data
+		TArray<FVoice*> VoiceObjects;
+		FEntityManager EntityManager;
+
+		FUnrealAudioModule* AudioModule;
+		FVoiceManagerSettings Settings;
+
+		// Audio thread data
+
+		// Handle index to voice pointer
+		TMap<uint32, uint32> HandleToIndex;
+
+		// VoiceDataIndex
+		TArray<FVoiceHandle> VoiceHandles;
+		TArray<EVoicePlayingState::Type> PlayingStates;
+		TArray<EVoiceState::Type> States;
+		TArray<TSharedPtr<ISoundFile>> SoundFiles;
+		FPitchManager PitchManager;
+		FVolumeManager VolumeManager;
+		TArray<uint32> Flags;
+		TQueue<uint32> FreeDataIndicies;
+		
+		FVoiceMixer VoiceMixer;
+		int32 NumVirtualVoices;
+		int32 NumRealVoices;
+	};
+
+}
+
+#endif // #if ENABLE_UNREAL_AUDIO
\ No newline at end of file
