Unreal Audio WIP

diff --git a/Engine/Source/Runtime/UnrealAudio/Private/UnrealAudioVoiceMixer.cpp b/Engine/Source/Runtime/UnrealAudio/Private/UnrealAudioVoiceMixer.cpp
new file mode 100644
--- /dev/null
+++ b/Engine/Source/Runtime/UnrealAudio/Private/UnrealAudioVoiceMixer.cpp
@@ -0,0 +1,248 @@
+// Copyright 1998-2015 Epic Games, Inc. All Rights Reserved.
+
+#include "UnrealAudioPrivate.h"
+#include "UnrealAudioModule.h"
+#include "UnrealAudioVoiceInternal.h"
+#include "UnrealAudioVoiceMixer.h"
+
+#if ENABLE_UNREAL_AUDIO
+
+namespace UAudio
+{
+	FVoiceMixer::FVoiceMixer(FUnrealAudioModule* InAudioModule, FVoiceManager* InVoiceManager)
+		: AudioModule(InAudioModule)
+		, VoiceManager(InVoiceManager)
+		, NumVoicesPerDecoder(0)
+		, MaxVoices(0)
+		, NumActiveVoices(0)
+		, DeviceThreadCommandQueue(500)
+		, AudioSystemCommandQueue(1000)
+#if ENABLE_UNREAL_AUDIO_THREAD_DEBUGGING
+		, AudioThreadId(INDEX_NONE)
+		, DeviceThreadId(INDEX_NONE)
+#endif
+	{
+	}
+
+	FVoiceMixer::~FVoiceMixer()
+	{
+	}
+
+	void FVoiceMixer::Init(const struct FVoiceManagerSettings& Settings)
+	{
+		InitAudioSystemThreadId();
+
+		// Initialize the sound file decoders used by voice manager
+		ControlUpdateRateSeconds = Settings.ControlUpdateRateSeconds;
+		NumVoicesPerDecoder = Settings.MaxVoiceCount / Settings.NumDecoders;
+		SoundFileDecoders.Init(FSoundFileDecoder(AudioModule), Settings.NumDecoders);
+		for (uint32 i = 0; i < Settings.NumDecoders; ++i)
+		{
+			SoundFileDecoders[i].Init(Settings.DecoderSettings, NumVoicesPerDecoder);
+		}
+		// initialize the data needed for mixing
+		VoiceData.Init(FVoiceDataEntry(), Settings.MaxVoiceCount);
+
+		for (int32 i = 0; i < Settings.MaxVoiceCount; ++i)
+		{
+			FreeVoiceIndices.Enqueue(i);
+		}
+
+	}
+
+	int32 FVoiceMixer::InitEntry(const FVoiceMixerVoiceData& Data)
+	{
+// 		int32 VoiceIndex;
+// 		if (FreeVoiceIndices.Dequeue(VoiceIndex))
+// 		{
+// 			DeviceThreadCommandQueue.Enqueue(FCommand(EAudioDeviceThreadMixCommand::VOICE_INIT_ENTRY, VoiceIndex, Data.NumChannels, Data.FrameRate, Data.VolumeProduct, Data.PitchProduct));
+// 
+// // 			Entry.NumChannels = Data.NumChannels;
+// // 			Entry.FrameRate = Data.FrameRate;
+// // 			Entry.CurrentFrame = INDEX_NONE;
+// // 			Entry.Volume.Init(Data.VolumeProduct);
+// // 			Entry.Pitch.Init(Data.PitchProduct);
+// // 			Entry.SampleRateConverter.Init(Data.PitchProduct, Data.NumChannesl);
+// 
+// 			// Now flag this voice entry as pending activation.. it won't activate until the next update
+// 			// from the audio device
+// 			return VoiceIndex;
+// 		}
+
+		return INDEX_NONE;
+	}
+
+	void FVoiceMixer::ReleaseEntry(uint32 VoiceIndex)
+	{
+// 		FreeVoiceIndices.Enqueue(VoiceIndex);
+// 		DeviceThreadCommandQueue.Enqueue(FCommand(EAudioDeviceThreadMixCommand::VOICE_RELEASE_ENTRY, VoiceIndex, VoiceIndex));
+	}
+
+	void FVoiceMixer::UpdateDeviceThread(struct FCallbackInfo& CallbackInfo)
+	{
+// 		InitDeviceThreadId();
+// 
+// 		int32 FrameRate = CallbackInfo.FrameRate;
+// 		float* OutBuffer = CallbackInfo.OutBuffer;
+// 		uint32 NumFrames = CallbackInfo.NumFrames;
+// 		uint32 NumChannels = CallbackInfo.NumChannels;
+// 		const TArray<ESpeaker::Type>& OutputSpeakers = CallbackInfo.OutputSpeakers;
+// 		double StreamTime = CallbackInfo.StreamTime;
+// 
+// 		// How many frames to interpolate any parameter changes from system thread
+// 		int32 InterpolationFrames = FrameRate * ControlUpdateRateSeconds;
+// 		check(InterpolationFrames > 0);
+// 
+// 		// Update any pending messages
+// 		PumpDeviceThreadMessages();
+// 
+// 		uint32 VoiceCount = 0;
+// 
+// 		for (uint32 VoiceIndex = 0; VoiceIndex < MaxVoices && VoiceCount < NumActiveVoices.GetValue(); ++VoiceIndex)
+// 		{
+// 			if (!VoiceData.IsActive[VoiceIndex])
+// 			{
+// 				continue;
+// 			}
+// 
+// 			++VoiceCount;
+// 
+// 			// Get voice data index (index in the overall data array of this voice)
+// 			uint32 VoiceDataIndex = VoiceData.VoiceDataIndex[VoiceIndex];
+// 			check(VoiceDataIndex != INVALID_INDEX);
+// 
+// 			// Get the volume and pitch information from voice manager
+// 			float VolumeProduct = 0.0f;
+// 			VoiceManager->GetVolumeProduct(VoiceDataIndex, VolumeProduct);
+// 
+// 			float PitchProduct = 1.0f;
+// 			VoiceManager->GetPitchProduct(VoiceDataIndex, PitchProduct);
+// 
+// 			// Get current volume data
+// 			float VolumeTarget = VoiceData.VolumeTarget[VoiceIndex];
+// 			float VolumeCurrent = VoiceData.VolumeCurrent[VoiceIndex];
+// 
+// 			// Update the target volume and volume increment if different
+// 			if (VolumeProduct != VolumeTarget)
+// 			{
+// 				VoiceData.VolumeTarget[VoiceIndex] = VolumeProduct;
+// 				VoiceData.VolumeDelta[VoiceIndex] = (VolumeProduct - VolumeCurrent) / InterpolationFrames;
+// 				VoiceData.VolumeFrameCount[VoiceIndex] = 0;
+// 			}
+// 
+// 			// Get the current pitch data
+// 			float PitchTarget = VoiceData.PitchTarget[VoiceIndex];
+// 			float PitchCurrent = VoiceData.PitchCurrent[VoiceData];
+// 
+// 			if (PitchProduct != PitchTarget)
+// 			{
+// 				VoiceData.PitchTarget[VoiceIndex] = PitchTarget;
+// 				VoiceData.PitchDelta[VoiceIndex] = (PitchProduct - PitchCurrent) / InterpolationFrames;
+// 				VoiceData.PitchFrameCount[VoiceIndex] = 0;
+// 			}
+// 
+// 			for (uint32 FrameIndex = 0; FrameIndex < NumFrames; ++FrameIndex)
+// 			{
+// 				// Figure out what the next frame will be based on the current pitch
+// 				double NextFrame = VoiceData.CurrentFrame[VoiceIndex] + (double)VoiceData.PitchCurrent[VoiceIndex];
+// 
+// 				// Figure out how many frames we need to get from the decoder to account for the current pitch value
+// 				uint32 DecodedFramesToRequest = 0;
+// 				if ((int32)VoiceData.CurrentFrame[VoiceIndex] != (int32)NextFrame)
+// 				{
+// 					DecodedFramesToRequest = (int32)NextFrame - (int32)VoiceData.CurrentFrame[VoiceIndex];
+// 				}
+// 
+// 				if (DecodedFramesToRequest != 0)
+// 				{
+// 					uint32 NumSamples = DecodedFramesToRequest * VoiceData.NumChannels[VoiceIndex];
+// 
+// 					// Reset the decode buffer to be the number of samples we're going to get this frame
+// 					DecodeBuffer.Reset(NumSamples);
+// 
+// 					// Now fill the buffer from the decoder with the number of samples
+// 					FillDecodedBuffer(VoiceIndex, DecodeBuffer);
+// 				}
+// 
+// 				// Now we should either have a new decode buffer or the previous decode buffer (which will happen if pitch is less than 1.0)
+// 
+// 
+// 				// Update the current volume and pitch interpolation values
+// 				if (VoiceData.PitchFrameCount[VoiceIndex] < InterpolationFrames)
+// 				{
+// 					VoiceData.PitchFrameCount[VoiceIndex]++;
+// 					VoiceData.VolumeCurrent[VoiceIndex] += VoiceData.VolumeDelta[VoiceIndex];
+// 				}
+// 
+// 				if (VoiceData.VolumeFrameCount[VoiceIndex] < InterpolationFrames)
+// 				{
+// 					VoiceData.VolumeFrameCount[VoiceIndex]++;
+// 					VoiceData.PitchCurrent[VoiceIndex] += VoiceData.PitchDelta[VoiceIndex];
+// 				}
+// 			}
+// 		}
+	}
+
+	uint32 FVoiceMixer::GetNumActiveVoices() const
+	{
+		return 0;
+//		return NumActiveVoices.GetValue();
+	}
+
+	void FVoiceMixer::FillDecodedBuffer(uint32 VoiceIndex, TArray<float>& InDecodeBuffer)
+	{
+// 		// first get which decoder to use given the index
+// 		check(NumVoicesPerDecoder > 0);
+// 		uint32 DecoderIndex = VoiceIndex / NumVoicesPerDecoder;
+// 		uint32 DecoderVoiceIndex = VoiceIndex % NumVoicesPerDecoder;
+// 
+// 		FSoundFileDecoder& Decoder = SoundFileDecoders[DecoderIndex];
+// 
+// 		Decoder.GetDecodedAudioData(DecoderVoiceIndex, InDecodeBuffer);
+	}
+
+	void FVoiceMixer::InitAudioSystemThreadId()
+	{
+#if ENABLE_UNREAL_AUDIO_THREAD_DEBUGGING
+		AudioThreadId = FPlatformTLS::GetCurrentThreadId();
+#endif // #if ENABLE_UNREAL_AUDIO_THREAD_DEBUGGING
+	}
+
+	void FVoiceMixer::InitDeviceThreadId()
+	{
+#if ENABLE_UNREAL_AUDIO_THREAD_DEBUGGING
+		if (DeviceThreadId == INDEX_NONE)
+		{
+			DeviceThreadId = FPlatformTLS::GetCurrentThreadId();
+		}
+#endif // #if ENABLE_UNREAL_AUDIO_THREAD_DEBUGGING
+	}
+
+	void FVoiceMixer::CheckAudioSystemThread()
+	{
+#if ENABLE_UNREAL_AUDIO_THREAD_DEBUGGING
+		uint32 CurrentThreadId = FPlatformTLS::GetCurrentThreadId();
+		checkf(CurrentThreadId == AudioThreadId,
+			   TEXT("Function called on wrong thread with id '%d' but supposed to be called on audio thread (id=%d)."),
+			   CurrentThreadId,
+			   AudioThreadId
+			   );
+#endif // ENABLE_UNREAL_AUDIO_THREAD_DEBUGGING
+	}
+
+	void FVoiceMixer::CheckDeviceThread()
+	{
+#if ENABLE_UNREAL_AUDIO_THREAD_DEBUGGING
+		uint32 CurrentThreadId = FPlatformTLS::GetCurrentThreadId();
+		checkf(CurrentThreadId == DeviceThreadId,
+			   TEXT("Function called on wrong thread with id '%d' but supposed to be called on audio device thread (id=%d)."),
+			   CurrentThreadId,
+			   AudioThreadId
+			   );
+#endif // ENABLE_UNREAL_AUDIO_THREAD_DEBUGGING
+	}
+
+}
+
+#endif // #if ENABLE_UNREAL_AUDIO
+
diff --git a/Engine/Source/Runtime/UnrealAudio/Private/UnrealAudioVoiceMixer.h b/Engine/Source/Runtime/UnrealAudio/Private/UnrealAudioVoiceMixer.h
new file mode 100644
--- /dev/null
+++ b/Engine/Source/Runtime/UnrealAudio/Private/UnrealAudioVoiceMixer.h
@@ -0,0 +1,154 @@
+#pragma once 
+
+#include "UnrealAudioPrivate.h"
+#include "UnrealAudioModule.h"
+#include "UnrealAudioTypes.h"
+#include "UnrealAudioVoice.h"
+#include "UnrealAudioVoiceInternal.h"
+#include "UnrealAudioUtilities.h"
+#include "UnrealAudioVoiceManager.h"
+#include "UnrealAudioDecode.h"
+#include "UnrealAudioSampleRateConverter.h"
+
+#if ENABLE_UNREAL_AUDIO
+
+namespace UAudio
+{
+	struct FVoiceMixerVoiceData
+	{
+		uint32 VoiceManagerIndex;
+		float VolumeProduct;
+		float PitchProduct;
+		TSharedPtr<ISoundFile> SoundFile;
+	};
+
+	struct FVoiceParamData
+	{
+		float Current;
+		float Target;
+		float Delta;
+		uint32 FrameCount;
+
+		FVoiceParamData()
+			: Current(0.0f)
+			, Target(0.0f)
+			, Delta(0.0f)
+			, FrameCount(0)
+		{}
+
+		void Init(float Value)
+		{
+			Current = Value;
+			Target = Value;
+			Delta = 0.0f;
+			FrameCount = 0;
+		}
+
+	};
+
+	struct FVoiceDataEntry
+	{
+		uint32 NumChannels;
+		uint32 FrameRate;
+		double CurrentFrame;
+		bool IsActive;
+		FVoiceParamData Volume;
+		FVoiceParamData Pitch;
+		FSampleRateConverter SampleRateConverter;
+
+		FVoiceDataEntry()
+			: NumChannels(0)
+			, FrameRate(0)
+			, CurrentFrame(0.0)
+			, IsActive(false)
+		{}
+
+	};
+
+	namespace EAudioDeviceThreadMixCommand
+	{
+		enum Type
+		{
+			NONE = 0,
+
+			VOICE_INIT_ENTRY,
+			VOICE_RELEASE_ENTRY,
+			VOICE_SET_VOLUME_PRODUCT,
+			VOICE_SET_PITCH_PRODUCT,
+			VOICE_SET_LISTENER_RELATIVE_ANGLE,
+
+		};
+	}
+
+	namespace EAudioSystemThreadMixCommand
+	{
+		enum Type
+		{
+			NONE = 0,
+
+			VOICE_DONE,
+		};
+	}
+
+	class FVoiceMixer
+	{
+	public:
+		FVoiceMixer(FUnrealAudioModule* InAudioModule, class FVoiceManager* InVoiceManager);
+		~FVoiceMixer();
+
+		void Init(const struct FVoiceManagerSettings& Settings);
+
+		int32 InitEntry(const FVoiceMixerVoiceData& Data);
+		void ReleaseEntry(uint32 VoiceIndex);
+
+		void UpdateDeviceThread(struct FCallbackInfo& CallbackInfo);
+
+		void UpdateAudioSystem();
+
+		uint32 GetNumActiveVoices() const;
+
+	private:
+		void PumpDeviceThreadMessages();
+		void FillDecodedBuffer(uint32 VoiceIndex, TArray<float>& InDecodeBuffer);
+
+		void InitAudioSystemThreadId();
+		void InitDeviceThreadId();
+		void CheckAudioSystemThread();
+		void CheckDeviceThread();
+
+		FUnrealAudioModule* AudioModule;
+
+		// Parent voice manager object
+		FVoiceManager* VoiceManager;
+
+		// Asynchronous decoders of audio file data
+		TArray<FSoundFileDecoder> SoundFileDecoders;
+		uint32 NumVoicesPerDecoder;
+
+		uint32 MaxVoices;
+		FThreadSafeCounter NumActiveVoices;
+		float ControlUpdateRateSeconds;
+
+		// Scratch buffer we are going to mix audio into from decoders
+		TArray<float> DecodeBuffer;
+
+		TQueue<int32> FreeVoiceIndices;
+
+		TCommandQueue<FCommand> DeviceThreadCommandQueue;
+
+		TCommandQueue<FCommand> AudioSystemCommandQueue;
+
+		TArray<FVoiceDataEntry> VoiceData;
+
+
+#if ENABLE_UNREAL_AUDIO_THREAD_DEBUGGING
+		uint32 AudioThreadId;
+		uint32 DeviceThreadId;
+#endif // #if ENABLE_UNREAL_AUDIO_THREAD_DEBUGGING
+
+	};
+
+
+}
+
+#endif // #if ENABLE_UNREAL_AUDIO
\ No newline at end of file
diff --git a/Engine/Source/Runtime/UnrealAudio/Private/UnrealAudioVolumeManager.cpp b/Engine/Source/Runtime/UnrealAudio/Private/UnrealAudioVolumeManager.cpp
new file mode 100644
--- /dev/null
+++ b/Engine/Source/Runtime/UnrealAudio/Private/UnrealAudioVolumeManager.cpp
@@ -0,0 +1,188 @@
+// Copyright 1998-2015 Epic Games, Inc. All Rights Reserved.
+
+#include "UnrealAudioPrivate.h"
+#include "UnrealAudioModule.h"
+#include "UnrealAudioVolumeManager.h"
+
+#if ENABLE_UNREAL_AUDIO
+
+namespace UAudio
+{
+	FVolumeManager::FVolumeManager(FUnrealAudioModule* InAudioModule)
+		: AudioModule(InAudioModule)
+		, EntryCount(0)
+		, LastTimeSec(0.0)
+	{
+		check(AudioModule);
+	}
+
+	void FVolumeManager::Init(uint32 NumElements)
+	{
+		Baseline.Init(-1.0f, NumElements);
+		Attenuation.Init(1.0f, NumElements);
+		VolumeProduct.Init(1.0f, NumElements);
+		FadeScale.Init(NumElements);
+		DynamicScale.Init(NumElements);
+		PriorityWeight.Init(1.0f, NumElements);
+		VolumeWeightedPriority.Init(1.0f, NumElements);
+		EmitterHandle.Init(FEmitterHandle(), NumElements);
+	}
+
+	void FVolumeManager::Update()
+	{
+		double CurrentTimeSec = AudioModule->GetCurrentTimeSec();
+		LastTimeSec = CurrentTimeSec;
+
+		// If there are no volumes to update, then no need to update anything
+		if (EntryCount == 0)
+		{
+			return;
+		}
+
+		int32 UpdateCount = 0;
+		float DeltaTime = 0.0f;
+		float Fraction = 0.0f;
+		float Result = 0.0f;
+		int32 TotalCount = Baseline.Num();
+
+		// Update all dynamic volume parameters in one parallelized loop.
+		// TODO: implement SSE
+
+		for (int32 Index = 0; Index < TotalCount && UpdateCount < EntryCount; ++Index)
+		{
+			// If the value of baseline is less than 0.0f, it's been a released volume entry
+			if (Baseline[Index] < 0.0f)
+			{
+				continue;
+			}
+
+			++UpdateCount;
+			float Product =	Baseline[Index] * 
+							Attenuation[Index] * 
+							DynamicScale.Compute(Index, CurrentTimeSec) * 
+							FadeScale.Compute(Index, CurrentTimeSec);
+
+			DEBUG_AUDIO_CHECK(Product >= 0.0f && Product <= 1.0f);
+			VolumeProduct[Index] = Product;
+		}
+
+		// Reset the volume priority sorter
+		SortedVoices.Reset();
+		UpdateCount = 0;
+		VoiceComparePredicate Predicate;
+		for (int32 Index = 0; Index < TotalCount && UpdateCount < EntryCount; ++Index)
+		{
+			if (Baseline[Index] < 0.0f)
+			{
+				continue;
+			}
+			++UpdateCount;
+			float Temp = VolumeProduct[Index] * PriorityWeight[Index];
+			VolumeWeightedPriority[Index] = Temp;
+			SortedVoices.HeapPush(FSortedVoiceEntry(Index, Temp), Predicate);
+		}
+
+		// Make sure we updated everything
+		DEBUG_AUDIO_CHECK(UpdateCount == EntryCount);
+	}
+
+	void FVolumeManager::InitializeEntry(uint32 VoiceDataIndex, const FVolumeInitParam& Params)
+	{
+		++EntryCount;
+		DEBUG_AUDIO_CHECK(EntryCount< Baseline.Num());
+		DEBUG_AUDIO_CHECK(VoiceDataIndex < (uint32)Baseline.Num());
+		DEBUG_AUDIO_CHECK(Params.VolumeProduct >= 0.0f && Params.VolumeProduct <= 1.0f);
+
+		Baseline[VoiceDataIndex] = Params.BaseVolume;
+		Attenuation[VoiceDataIndex] = Params.VolumeAttenuation;
+		VolumeProduct[VoiceDataIndex] = Params.VolumeProduct;
+		EmitterHandle[VoiceDataIndex] = Params.EmitterHandle;
+		PriorityWeight[VoiceDataIndex] = Params.PriorityWeight;
+
+		DynamicScale.InitEntry(VoiceDataIndex);
+		DynamicScale.SetValue(VoiceDataIndex, Params.VolumeScale, LastTimeSec, Params.VolumeScaleDeltaTime);
+
+		FadeScale.InitEntry(VoiceDataIndex);
+	}
+
+	void FVolumeManager::ReleaseEntry(uint32 VoiceDataIndex)
+	{
+		DEBUG_AUDIO_CHECK(VoiceDataIndex < (uint32)Baseline.Num());
+
+		--EntryCount;
+		DEBUG_AUDIO_CHECK(EntryCount >= 0);
+
+		// Setting baseline to -1.0f is releasing the volume data at that index
+		Baseline[VoiceDataIndex] = -1.0f;
+	}
+
+	void FVolumeManager::SetAttenuation(uint32 VoiceDataIndex, float InAttenuation)
+	{
+		DEBUG_AUDIO_CHECK(VoiceDataIndex < (uint32)Baseline.Num());
+		DEBUG_AUDIO_CHECK(Baseline[VoiceDataIndex] > 0.0f);
+		Attenuation[VoiceDataIndex] = InAttenuation;
+	}
+
+	void FVolumeManager::SetDynamicVolumeScale(uint32 VoiceDataIndex, float InVolume, float DeltaTimeSec)
+	{
+		DEBUG_AUDIO_CHECK(VoiceDataIndex < (uint32)Baseline.Num());
+		DEBUG_AUDIO_CHECK(Baseline[VoiceDataIndex] > 0.0f);
+		DEBUG_AUDIO_CHECK(InVolume >= 0.0f && InVolume <= 1.0f);
+		DynamicScale.SetValue(VoiceDataIndex, InVolume, LastTimeSec, DeltaTimeSec);
+	}
+
+	void FVolumeManager::SetFadeOut(uint32 VoiceDataIndex, float FadeTimeSec)
+	{
+		DEBUG_AUDIO_CHECK(VoiceDataIndex < (uint32)Baseline.Num());
+		DEBUG_AUDIO_CHECK(Baseline[VoiceDataIndex] > 0.0f);
+		FadeScale.SetValue(VoiceDataIndex, 0.0f, LastTimeSec, FadeTimeSec);
+	}
+
+	bool FVolumeManager::IsFadedOut(uint32 VoiceDataIndex) const
+	{
+		DEBUG_AUDIO_CHECK(VoiceDataIndex < (uint32)Baseline.Num());
+		DEBUG_AUDIO_CHECK(Baseline[VoiceDataIndex] > 0.0f);
+		return FadeScale.EndValue[VoiceDataIndex] == 0.0f && FadeScale.bIsDone[VoiceDataIndex];
+	}
+
+	float FVolumeManager::GetVolumeScale(uint32 VoiceDataIndex) const
+	{
+		DEBUG_AUDIO_CHECK(VoiceDataIndex < (uint32)Baseline.Num());
+		DEBUG_AUDIO_CHECK(Baseline[VoiceDataIndex] > 0.0f);
+
+		return DynamicScale.CurrentValue[VoiceDataIndex];
+	}
+
+	float FVolumeManager::GetVolumeFade(uint32 VoiceDataIndex) const
+	{
+		DEBUG_AUDIO_CHECK(VoiceDataIndex < (uint32)Baseline.Num());
+		DEBUG_AUDIO_CHECK(Baseline[VoiceDataIndex] > 0.0f);
+
+		return FadeScale.CurrentValue[VoiceDataIndex];
+	}
+
+	float FVolumeManager::GetVolumeProduct(uint32 VoiceDataIndex) const
+	{
+		DEBUG_AUDIO_CHECK(VoiceDataIndex < (uint32)Baseline.Num());
+		DEBUG_AUDIO_CHECK(Baseline[VoiceDataIndex] > 0.0f);
+
+		return VolumeProduct[VoiceDataIndex];
+	}
+
+	float FVolumeManager::GetVolumeAttenuation(uint32 VoiceDataIndex) const
+	{
+		DEBUG_AUDIO_CHECK(VoiceDataIndex < (uint32)Baseline.Num());
+		DEBUG_AUDIO_CHECK(Baseline[VoiceDataIndex] > 0.0f);
+
+		return Attenuation[VoiceDataIndex];
+	}
+
+	const TArray<FSortedVoiceEntry>& FVolumeManager::GetSortedVoices() const
+	{
+		return SortedVoices;
+	}
+
+}
+
+#endif // #if ENABLE_UNREAL_AUDIO
+
diff --git a/Engine/Source/Runtime/UnrealAudio/Private/UnrealAudioVolumeManager.h b/Engine/Source/Runtime/UnrealAudio/Private/UnrealAudioVolumeManager.h
new file mode 100644
--- /dev/null
+++ b/Engine/Source/Runtime/UnrealAudio/Private/UnrealAudioVolumeManager.h
@@ -0,0 +1,122 @@
+#pragma once 
+
+#include "UnrealAudioModule.h"
+#include "UnrealAudioTypes.h"
+#include "UnrealAudioVoice.h"
+#include "UnrealAudioPrivate.h"
+#include "UnrealAudioVoiceInternal.h"
+#include "UnrealAudioUtilities.h"
+
+#if ENABLE_UNREAL_AUDIO
+
+namespace UAudio
+{
+	/**
+	 * A struct containing a index (into voice entity array) and a computed volume product.
+	 */
+
+	struct FSortedVoiceEntry
+	{
+		/**
+		 * Constructor.
+		 *
+		 * @param	Index   	Zero-based index of the.
+		 * @param	InVolume	The in volume.
+		 */
+
+		explicit FSortedVoiceEntry(uint32 InIndex, float InPriorityWeightedVolume)
+			: Index(InIndex)
+			, PriorityWeightedVolume(InPriorityWeightedVolume)
+		{}
+
+		/** Zero-based index of the object. */
+		uint32 Index;
+
+		/** The priority-weighted volume. */
+		float PriorityWeightedVolume;
+	};
+
+	struct FVolumeInitParam
+	{
+		FEmitterHandle EmitterHandle;
+		float BaseVolume;
+		float VolumeScale;
+		float VolumeScaleDeltaTime;
+		float VolumeProduct;
+		float VolumeAttenuation;
+		float PriorityWeight;
+	};
+
+	/** Manager for volumes. */
+	class FVolumeManager
+	{
+	public:
+		FVolumeManager(FUnrealAudioModule* InAudioModule);
+
+		void Init(uint32 NumElements);
+		void Update();
+		
+		void InitializeEntry(uint32 VoiceDataIndex, const FVolumeInitParam& Params);
+		void ReleaseEntry(uint32 VoiceDataIndex);
+
+		void SetAttenuation(uint32 VoiceDataIndex, float Attenuation);
+		void SetDynamicVolumeScale(uint32 VoiceDataIndex, float VolumeScale, float DeltaTimeSec);
+		void SetFadeOut(uint32 VoiceDataIndex, float FadeTimeSec);
+
+		bool IsFadedOut(uint32 VoiceDataIndex) const;
+		float GetVolumeScale(uint32 VoiceDataIndex) const;
+		float GetVolumeFade(uint32 VoiceDataIndex) const;
+		float GetVolumeProduct(uint32 VoiceDataIndex) const;
+		float GetVolumeAttenuation(uint32 VoiceDataIndex) const;
+		const TArray<FSortedVoiceEntry>& GetSortedVoices() const;
+
+	private:
+
+		/** The audio module. */
+		FUnrealAudioModule* AudioModule;
+
+		/** Number of entries. */
+		int32 EntryCount;
+
+		/** The last time the object was updated. */
+		double LastTimeSec;
+
+		/** The baseline volume. */
+		TArray<float> Baseline;
+
+		/** The attenuation. */
+		TArray<float> Attenuation;
+
+		/** The volume product. */
+		TArray<float> VolumeProduct;
+
+		/** The volume scale. */
+		FDynamicParamData DynamicScale;
+
+		/** The volume fade scale. */
+		FDynamicParamData FadeScale;
+
+		/** The priority weight of the voice, used to sort voices. */
+		TArray<float> PriorityWeight;
+
+		TArray<float> VolumeWeightedPriority;
+
+		TArray<FEmitterHandle> EmitterHandle;
+
+		/** A volume compare predicate. */
+		struct VoiceComparePredicate
+		{
+			bool operator()(const FSortedVoiceEntry& A, const FSortedVoiceEntry& B) const
+			{
+				return A.PriorityWeightedVolume > B.PriorityWeightedVolume;
+			}
+		};
+
+		/** Array of sorted indices by volume. */
+		TArray<FSortedVoiceEntry> SortedVoices;
+	};
+
+
+}
+
+#endif // #if ENABLE_UNREAL_AUDIO
diff --git a/Engine/Source/Runtime/UnrealAudio/Public/UnrealAudioEmitter.h b/Engine/Source/Runtime/UnrealAudio/Public/UnrealAudioEmitter.h
new file mode 100644
--- /dev/null
+++ b/Engine/Source/Runtime/UnrealAudio/Public/UnrealAudioEmitter.h
@@ -0,0 +1,50 @@
+// Copyright 1998-2015 Epic Games, Inc. All Rights Reserved.
+
+#pragma once 
+
+#include "UnrealAudioModule.h"
+#include "UnrealAudioTypes.h"
+
+namespace UAudio
+{
+	/** IEmitter interface. */
+	class UNREALAUDIO_API IEmitter
+	{
+	public:
+		/** Destructor. */
+		virtual ~IEmitter() {}
+
+		/**
+		 * Sets the emitter's position.
+		 *
+		 * @param	InPosition	The position to set the emitter to.
+		 *
+		 * @return	An ESystemError::Type.
+		 */
+		virtual ESystemError::Type SetPosition(const FVector& InPosition) = 0;
+
+		/**
+		 * Gets the emitter's position.
+		 *
+		 * @param [out]		OutPosition	The out position.
+		 *
+		 * @return	An ESystemError::Type.
+		 */
+		virtual ESystemError::Type GetPosition(FVector& OutPosition) = 0;
+
+		/**
+		 * Releases the emitter object.
+		 *
+		 * @return	An ESystemError::Type.
+		 */
+		virtual ESystemError::Type Release() = 0;
+
+		/**
+		 * Gets the emitter identifier.
+		 *
+		 * @return	The identifier.
+		 */
+		virtual uint32 GetId() const = 0;
+	};
+}
+
diff --git a/Engine/Source/Runtime/UnrealAudio/Public/UnrealAudioModule.h b/Engine/Source/Runtime/UnrealAudio/Public/UnrealAudioModule.h
--- a/Engine/Source/Runtime/UnrealAudio/Public/UnrealAudioModule.h
+++ b/Engine/Source/Runtime/UnrealAudio/Public/UnrealAudioModule.h
@@ -5,50 +5,140 @@
 #include "ModuleInterface.h"
 #include "UnrealAudioTypes.h"
 #include "UnrealAudioSoundFile.h"
+#include "UnrealAudioEmitter.h"
+#include "UnrealAudioVoice.h"
 
 namespace UAudio
 {
+	/** High-level unreal audio module for managing, playing, and updating audio in the unreal engine. */
 	class UNREALAUDIO_API IUnrealAudioModule :  public IModuleInterface
 	{
 	public:
+
+		/**
+		 * Destructor.
+		 */
 		virtual ~IUnrealAudioModule() {}
 
+		/**
+		 * Initializes this object.
+		 *
+		 * @return	true if it succeeds, false if it fails.
+		 */
 		virtual bool Initialize()
 		{
 			return false;
 		}
 
+		/**
+		 * Initializes this object.
+		 *
+		 * @param	DeviceModuleName	Name of the device module.
+		 *
+		 * @return	true if it succeeds, false if it fails.
+		 */
 		virtual bool Initialize(const FString& DeviceModuleName)
 		{
 			return false;
 		}
 
+		/** Updates this object. */
+		virtual void Update()
+		{
+		}
+
+		/** Shuts down this object and frees any resources it is using. */
 		virtual void Shutdown()
 		{
 		}
 
-		/** 
-		ImportSound
-		@param ImportSettings Settings to use to import a sound file. Settings file contains all information
-				needed to performt he import process.
-		@return A shared pointer to an ISoundFile object.
-		*/
-		virtual	TSharedPtr<ISoundFile> ImportSound(const FSoundFileImportSettings& ImportSettings)
+		virtual int32 GetNumBackgroundTasks() const
+		{
+			return 0;
+		}
+
+		virtual void ConvertSound(const FString& FilePath, const FString& OutputFilePath, const FSoundFileConvertFormat& ConvertFormat)
+		{
+
+		}
+
+		virtual void ConvertSound(const FString& FilePath, const FString& OutputFilePath)
+		{
+
+		}
+
+		virtual TSharedPtr<ISoundFile> LoadSoundFile(const FName& Name, TArray<uint8>& InBulkData)
+		{
+			return nullptr;
+		}
+
+		virtual TSharedPtr<ISoundFile> LoadSoundFile(const FName& Path, bool bLoadAsync = true)
+		{
+			return nullptr;
+		}
+
+		virtual TSharedPtr<ISoundFile>	StreamSoundFile(const FName& Path, bool bLoadAsync = true)
+		{
+			return nullptr;
+		}
+
+		/** Returns the number of sound files currently streamed */
+		virtual int32 GetNumSoundFilesStreamed() const
+		{
+			return 0;
+		}
+
+		/** Returns the number of sound files currently loaded */
+		virtual int32 GetNumSoundFilesLoaded() const
+		{
+			return 0;
+		}
+
+		/** Returns the number of total bytes of sound files loaded */
+		virtual int32 GetSoundFileNumBytes() const
+		{
+			return 0;
+		}
+
+		/** Returns the percentage of loaded memory as fraction of target memory limit */
+		virtual float GetSoundFilePercentageOfTargetMemoryLimit() const
+		{
+			return 0.0f;
+		}
+
+		/** Logs sound file information */
+		virtual void LogSoundFileMemoryInfo() const
+		{
+		}
+
+		/**
+		 * Emitter create.
+		 *
+		 * @return	An IEmitter object.
+		 */
+		virtual TSharedPtr<IEmitter> EmitterCreate()
 		{
 			return nullptr;
 		}
 
 		/**
-		ImportSound
-		@param ImportSettings Settings to use to import a sound file. Settings file contains all information
-		needed to performt he import process.
-		@return A handle to an ISoundFile object.
-		*/
-		virtual	void ExportSound(TSharedPtr<ISoundFileData> SoundFileData, const FString& ExportPath)
+		 * Voice create.
+		 *
+		 * @param	Params	A variable-length parameters list containing parameters.
+		 *
+		 * @return	An IVoice object.
+		 */
+		virtual TSharedPtr<IVoice> VoiceCreate(const FVoiceInitializationParams& Params)
 		{
+			return nullptr;
 		}
 
 
+		virtual class IUnrealAudioDeviceModule* GetDeviceModule()
+		{
+			return nullptr;
+		}
+
 	};
 }
 
diff --git a/Engine/Source/Runtime/UnrealAudio/Public/UnrealAudioSoundFile.h b/Engine/Source/Runtime/UnrealAudio/Public/UnrealAudioSoundFile.h
--- a/Engine/Source/Runtime/UnrealAudio/Public/UnrealAudioSoundFile.h
+++ b/Engine/Source/Runtime/UnrealAudio/Public/UnrealAudioSoundFile.h
@@ -7,16 +7,15 @@
 
 namespace UAudio
 {
-	/*
-	* ESoundFileError
-	**/
 	namespace ESoundFileError
 	{
 		enum Type
 		{
-			NONE,
+			NONE = 0,
 			INVALID_SOUND_FILE,
+			INVALID_SOUND_FILE_HANDLE,
 			BAD_ENCODING_QUALITY,
+			FAILED_TO_LOAD_BYTE_DATA,
 			ALREADY_OPENED,
 			ALREADY_HAS_DATA,
 			INVALID_DATA,
@@ -27,42 +26,42 @@ namespace UAudio
 			FAILED_TO_SEEK,
 			ALREADY_INTIALIZED,
 			LOADING,
+			INVALID_STATE,
 			UNKNOWN
 		};
 
-		/** @return the stringified version of the enum passed in */
 		inline const TCHAR* ToString(ESoundFileError::Type SoundFileError)
 		{
 			switch (SoundFileError)
 			{
-				case NONE:					return TEXT("NONE");
-				case INVALID_SOUND_FILE:	return TEXT("INVALID_SOUND_FILE");
-				case BAD_ENCODING_QUALITY:	return TEXT("BAD_ENCODING_QUALITY");
-				case ALREADY_OPENED:		return TEXT("ALREADY_OPENED");
-				case ALREADY_HAS_DATA:		return TEXT("ALREADY_HAS_DATA");
-				case INVALID_DATA:			return TEXT("INVALID_DATA");
-				case FILE_DOESNT_EXIST:		return TEXT("FILE_DOESNT_EXIST");
-				case INVALID_INPUT_FORMAT:	return TEXT("INVALID_INPUT_FORMAT");
-				case INVALID_CHANNEL_MAP:	return TEXT("INVALID_CHANNEL_MAP");
-				case FAILED_TO_OPEN:		return TEXT("FAILED_TO_OPEN");
-				case FAILED_TO_SEEK:		return TEXT("FAILED_TO_SEEK");
-				case ALREADY_INTIALIZED:	return TEXT("ALREADY_INTIALIZED");
-				case LOADING:				return TEXT("LOADING");
-
-				default:
-				case UNKNOWN:				return TEXT("UNKNOWN");
+				case NONE:						return TEXT("NONE");
+				case INVALID_SOUND_FILE:		return TEXT("INVALID_SOUND_FILE");
+				case INVALID_SOUND_FILE_HANDLE: return TEXT("INVALID_SOUND_FILE_HANDLE");
+				case BAD_ENCODING_QUALITY:		return TEXT("BAD_ENCODING_QUALITY");
+				case FAILED_TO_LOAD_BYTE_DATA:	return TEXT("FAILED_TO_LOAD_BYTE_DATA");
+				case ALREADY_OPENED:			return TEXT("ALREADY_OPENED");
+				case ALREADY_HAS_DATA:			return TEXT("ALREADY_HAS_DATA");
+				case INVALID_DATA:				return TEXT("INVALID_DATA");
+				case FILE_DOESNT_EXIST:			return TEXT("FILE_DOESNT_EXIST");
+				case INVALID_INPUT_FORMAT:		return TEXT("INVALID_INPUT_FORMAT");
+				case INVALID_CHANNEL_MAP:		return TEXT("INVALID_CHANNEL_MAP");
+				case FAILED_TO_OPEN:			return TEXT("FAILED_TO_OPEN");
+				case FAILED_TO_SEEK:			return TEXT("FAILED_TO_SEEK");
+				case ALREADY_INTIALIZED:		return TEXT("ALREADY_INTIALIZED");
+				case LOADING:					return TEXT("LOADING");
+				case INVALID_STATE:				return TEXT("INVALID_STATE");
+				default: case UNKNOWN:			return TEXT("UNKNOWN");
 			}
 		}
 	}
 
-	/*
-	* ESoundFileFormatFlags
-	* Specifies the major format type of the sound source.
-	* File formats are fully specified by a major/minor format.
-	*
-	* For example, a Ogg-Vorbis encoding would use:
-	* uint32 FormatFlags = ESoundFormatFlags::OGG | ESoundFormatFlags::VORBIS;
-	**/
+	/**
+	 * Specifies the major format type of the sound source.
+	 * File formats are fully specified by a major/minor format.
+	 *
+	 * For example, a Ogg-Vorbis encoding would use:
+	 * uint32 FormatFlags = ESoundFormatFlags::OGG | ESoundFormatFlags::VORBIS;
+	 */
 	namespace ESoundFileFormat
 	{
 		enum Flags
@@ -151,10 +150,10 @@ namespace UAudio
 	};
 
 	/*
-	* ESoundFileChannelMap
-	* These are separated from the device channel mappings purposefully since
-	* the enumeration may not exactly be the same as the output speaker mapping.
-	**/
+	 * Enumeration to specify a sound files intended output channel mapping.
+	 * @note These are separated from the device channel mappings purposefully since
+	 * the enumeration may not exactly be the same as the output speaker mapping.
+	 */
 	namespace ESoundFileChannelMap
 	{
 		enum Type
@@ -184,7 +183,6 @@ namespace UAudio
 			TOP_BACK_CENTER,
 		};
 
-		/** @return the stringified version of the enum passed in */
 		inline const TCHAR* ToString(ESoundFileChannelMap::Type ChannelMap)
 		{
 			switch (ChannelMap)
@@ -217,166 +215,84 @@ namespace UAudio
 		}
 	}
 
-	/*
-	* FSoundFileDescription
-	* Struct specifying a sound file description.
-	**/
+	/**
+	 * Specifes a sound file description.
+	 */
+
 	struct FSoundFileDescription
 	{
-		// The number of frames (interleaved samples) in the sound file.
+		/** The number of frames (interleaved samples) in the sound file. */
 		int64 NumFrames;
 
-		// The sample rate of the sound file.
+		/** The sample rate of the sound file. */
 		int32 SampleRate;
 
-		// The number of channels of the sound file.
+		/** The number of channels of the sound file. */
 		int32 NumChannels;
 
-		// The format flags of the sound file.
+		/** The format flags of the sound file. */
 		int32 FormatFlags;
 
-		// The number of sections of the sound file.
+		/** The number of sections of the sound file. */
 		int32 NumSections;
 
-		// Whether or not the sound file is seekable.
+		/** Whether or not the sound file is seekable. */
 		int32 bIsSeekable;
 	};
 
-	/**
-	* FSoundFileImportSettings
-	* Struct defining sound file import settings.
-	*/
-	struct UNREALAUDIO_API FSoundFileImportSettings
+	struct FSoundFileConvertFormat
 	{
-		// Full path to the sound file on disk.
-		FString SoundFilePath;
-
-		// Desired import format.
+		/** Desired convert format. */
 		int32 Format;
 
-		// Desired import sample rate.
+		/** Desired convert sample rate. */
 		uint32 SampleRate;
 
-		// For compression-type target formats that used an encoding quality (0.0 = low, 1.0 = high).
+		/** For compression-type target formats that used an encoding quality (0.0 = low, 1.0 = high). */
 		double EncodingQuality;
 
-		// Whether or not to peak-normalize the audio file during import.
+		/** Whether or not to peak-normalize the audio file during import. */
 		bool bPerformPeakNormalization;
 	};
 
-	/*
-	* ESoundFileState
-	**/
 	namespace ESoundFileState
 	{
 		enum Type
 		{
-			// Sound file hasn't begun loading and was just created.
-			UNINITIALIZED,
-
-			// Sound file has been initialized with data and is ready for playback
+			UNINITIALIZED = 0,
 			INITIALIZED,
-
-			// Sound file is importing.
-			IMPORTING,
-
-			// Sound file has finished importing.
-			IMPORTED,
-
-			// Sound file has an error. Called GetError() on ISoundFile to get error.
+			LOADING,
+			LOADED,
+			STREAMING,
+			WRITING,
 			HAS_ERROR,
 		};
 	}
 
 	/**
-	* FSoundFileData
-	* Data structure containing data for a sound file. Can be stored within a UAsset or loaded with a sound file.
-	*/
-	class UNREALAUDIO_API ISoundFileData
-	{
-	public:
-		virtual ~ISoundFileData() {}
-
-		/** Sets and gets the sound file path */
-		virtual const FString& GetFilePath() const = 0;
-		virtual void SetFilePath(const FString& InPath) = 0;
-		
-		/** Sets and gets the file description */
-		virtual const FSoundFileDescription& GetDescription() const = 0;
-		virtual void SetDescription(const FSoundFileDescription& InDescription) = 0;
-
-		/** Sets and gets the original file description */
-		virtual const FSoundFileDescription& GetOriginalDescription() const = 0;
-		virtual void SetOriginalDescription(const FSoundFileDescription& InDescription) = 0;
-
-		/** Sets and gets the channel map */
-		virtual const TArray<ESoundFileChannelMap::Type>& GetChannelMap() const = 0;
-		virtual void SetChannelMap(const TArray<ESoundFileChannelMap::Type>& InChannelMap) = 0;
-
-		/** Gets a const pointer to the internal buffer of data. */
-		virtual const uint8* GetData() const = 0;
-
-		/** Gets the size of the data */
-		virtual int64 GetDataSize() const = 0;
-
-		/** Sets the data buffer pointer and the size of the data buffer */
-		virtual void SetData(uint8* InDataPtr, int64 InDataSize) = 0;
-	};
-
-	/**
-	* ISoundFile
-	* Interface for sound file reading/writing.
-	*/
+	 * ISoundFile
+	 */
 	class UNREALAUDIO_API ISoundFile
 	{
 	public:
-		/** Virtual Destructor */
 		virtual ~ISoundFile() {}
-
-		/** Initializes the sound file with the given sound file data. */
-		virtual ESoundFileError::Type Initialize(TSharedPtr<ISoundFileData> InSoundFileData) = 0;
-
-		/** Returns the name of the sound file */
-		virtual ESoundFileError::Type GetName(FString& Name) = 0;
-
-		/** Gets the sound file data associated with this ISoundFile. */
-		virtual ESoundFileError::Type GetSoundFileData(TSharedPtr<ISoundFileData>& OutSoundFileData) = 0;
-
-		/** @return Returns the state of the sound file. */
-		virtual  ESoundFileState::Type GetState() const = 0;
-
-		/** @return Returns the last sound file error. */
+		virtual ESoundFileError::Type GetState(ESoundFileState::Type& OutState) const = 0;
 		virtual ESoundFileError::Type GetError() const = 0;
+		virtual ESoundFileError::Type GetId(uint32& OutId) const = 0;
+		virtual ESoundFileError::Type GetPath(FName& OutPath) const = 0;
+		virtual ESoundFileError::Type GetBulkData(TArray<uint8>** OutData) const = 0;
+		virtual ESoundFileError::Type GetDataSize(int32& DataSize) const = 0;
+		virtual ESoundFileError::Type GetDescription(FSoundFileDescription& OutDescription) const = 0;
+		virtual ESoundFileError::Type GetChannelMap(TArray<ESoundFileChannelMap::Type>& OutChannelMap) const = 0;
+		virtual ESoundFileError::Type IsStreamed(bool& bOutIsStreamed) const = 0;
 	};
 
-	/**
-	* Creates an empty sound file data object
-	* Creates an empty sound file (with no sound file data)
-	*/
-	UNREALAUDIO_API TSharedPtr<ISoundFileData> CreateSoundFileData();
+	UNREALAUDIO_API bool GetSoundFileDescription(const FString& FilePath, FSoundFileDescription& OutputDescription, TArray<ESoundFileChannelMap::Type>& OutChannelMap);
 
-	/**
-	* CreateSoundFile
-	* Creates an empty sound file (with no sound file data)
-	*/
-	UNREALAUDIO_API TSharedPtr<ISoundFile> CreateSoundFile();
+	UNREALAUDIO_API bool GetSoundFileDescription(const FString& FilePath, FSoundFileDescription& OutputDescription);
 
-	/**
-	* ImportSound
-	* Imports a sound given the import settings.
-	* The function immediately returns a ISoundFile reference, but that sound file will
-	* be in a loading state. If the sound file succeeds in importing, the sound file state
-	* will become Loaded, otherwise it will have an error that can be retrieved.
-	*/
-	UNREALAUDIO_API TSharedPtr<ISoundFile> ImportSound(const FSoundFileImportSettings& ImportSettings);
+	UNREALAUDIO_API bool GetFileExtensionForFormatFlags(int32 FormatFlags, FString& OutExtension);
 
-	/**
-	* ExportSound
-	* Exports a sound to the given path.
-	* The function immediately returns a ISoundFile reference, but that sound file will
-	* be in a loading state. If the sound file succeeds in importing, the sound file state
-	* will become Loaded, otherwise it will have an error that can be retrieved.
-	*/
-	UNREALAUDIO_API void ExportSound(TSharedPtr<ISoundFileData> SoundFileData, const FString& ExportPath);
+	UNREALAUDIO_API void GetSoundFileListInDirectory(FString& Directory, TArray<FString>& SoundFiles, bool bRecursive = true);
 
 }
diff --git a/Engine/Source/Runtime/UnrealAudio/Public/UnrealAudioTests.h b/Engine/Source/Runtime/UnrealAudio/Public/UnrealAudioTests.h
--- a/Engine/Source/Runtime/UnrealAudio/Public/UnrealAudioTests.h
+++ b/Engine/Source/Runtime/UnrealAudio/Public/UnrealAudioTests.h
@@ -44,32 +44,33 @@ namespace UAudio
 	bool UNREALAUDIO_API TestDeviceOutputNoisePan(double TestTime);
 
 	/**
-	* TestSourceImport
+	* TestSourceConvert
 	*
-	* Tests importing a sound file.
+	* Tests converting a sound file.
 	* @param ImportSettings A struct defining import settings.
 	* @return True if test succeeds.
 	*/
-	bool UNREALAUDIO_API TestSourceImport(const FSoundFileImportSettings& ImportSettings);
+	bool UNREALAUDIO_API TestSourceConvert(const FString& FilePath, const FSoundFileConvertFormat& ConvertFormat);
 
 	/**
-	* TestSourceImport
+	* TestEmitterManager
 	*
-	* Tests importing a sound file, then exporting it (with "export" appended to filename).
-	* @param ImportSettings A struct defining import settings.
+	* Tests the emitter manager with creating and updating emitters from one thread (simulating main thread) and processing from
+	* the audio system thread.
 	* @return True if test succeeds.
 	*/
-	bool UNREALAUDIO_API TestSourceImportExport(const FSoundFileImportSettings& ImportSettings);
+	bool UNREALAUDIO_API TestEmitterManager();
 
 	/**
-	* TestSourceImportFolder
+	* TestVoiceManager
 	*
-	* Tests importing all the sound files in a given folder.
-	* @param FolderPath Path to a folder of sound files.
-	* @param ImportSettings Settings to use for each sound file import (path is ignored)
+	* Tests the voice manager with creating and updating voices from one thread (simulating main thread) and processing from
+	* the audio system thread.
 	* @return True if test succeeds.
 	*/
-	bool UNREALAUDIO_API TestSourceImportFolder(FString& FolderPath, const FSoundFileImportSettings& ImportSettings);
+	bool UNREALAUDIO_API TestVoiceManager(const FString& FolderOrPath);
+
 
+	bool UNREALAUDIO_API TestSoundFileManager(const FString& FolderOrPath);
 }
 
diff --git a/Engine/Source/Runtime/UnrealAudio/Public/UnrealAudioTypes.h b/Engine/Source/Runtime/UnrealAudio/Public/UnrealAudioTypes.h
--- a/Engine/Source/Runtime/UnrealAudio/Public/UnrealAudioTypes.h
+++ b/Engine/Source/Runtime/UnrealAudio/Public/UnrealAudioTypes.h
@@ -4,21 +4,47 @@
 
 // Only enable unreal audio on windows or mac
 #if PLATFORM_WINDOWS
-#define ENABLE_UNREAL_AUDIO 1
+#define ENABLE_UNREAL_AUDIO 0
 #else
 #define ENABLE_UNREAL_AUDIO 0
 #endif
 
 #include "UnrealAudioModule.h"
 
-#if ENABLE_UNREAL_AUDIO
+// Turns on extra checks and asserts for debugging threading code
+#define ENABLE_UNREAL_AUDIO_THREAD_DEBUGGING		1
+#define ENABLE_UNREAL_AUDIO_EXTRA_DEBUG_CHECKS		1
 
 namespace UAudio
 {
-	/**
-	* EDeviceError
-	* The types of errors that can occur in IUnrealAudioDevice implementations.
-	*/
+	namespace ESystemError
+	{
+		enum Type
+		{
+			NONE,
+			INVALID_HANDLE,
+			COMMAND_QUEUE_FULL,
+			ALREADY_INITIALIZED,
+			NOT_INITIALIZED,
+			INVALID_ARGUMENTS,
+		};
+
+		static inline const TCHAR* ToString(ESystemError::Type SystemError)
+		{
+			switch (SystemError)
+			{
+				case NONE:						return TEXT("NONE");
+				case INVALID_HANDLE:			return TEXT("INVALID_HANDLE");
+				case COMMAND_QUEUE_FULL:		return TEXT("COMMAND_QUEUE_FULL");
+				case ALREADY_INITIALIZED:		return TEXT("ALREADY_INITIALIZED");
+				case NOT_INITIALIZED:			return TEXT("NOT_INITIALIZED");
+				case INVALID_ARGUMENTS:			return TEXT("INVALID_ARGUMENTS");
+				default:						check(false);
+			}
+			return TEXT("");
+		}
+	}
+
 	namespace EDeviceError
 	{
 		enum Type
@@ -31,7 +57,6 @@ namespace UAudio
 			THREAD,
 		};
 
-		/** @return the stringified version of the enum passed in */
 		static inline const TCHAR* ToString(EDeviceError::Type DeviceError)
 		{
 			switch (DeviceError)
@@ -48,49 +73,41 @@ namespace UAudio
 		}
 	}
 
-	/**
-	* EDeviceApi
-	* Implemented platform-Specific audio device APIs.
-	*/
 	namespace EDeviceApi
 	{
 		enum Type
 		{
-			WASAPI,				// Windows
-			XAUDIO2,			// Windows/XBox
-			NGS2,				// PS4
-			ALSA,				// Linux
-			CORE_AUDIO,			// Mac/IOS
-			OPENAL,				// Linux
-			HTML5,				// Web
-			DUMMY,				// DUMMY (no device api)
+			WASAPI,
+			XAUDIO2,
+			NGS2,
+			ALSA,
+			CORE_AUDIO,
+			OPENAL,
+			HTML5,
+			DUMMY,
 		};
 
-		/** @return the stringified version of the enum passed in */
 		static inline const TCHAR* ToString(EDeviceApi::Type DeviceApi)
 		{
 			switch (DeviceApi)
 			{
-			case WASAPI:		return TEXT("WASAPI");
-			case XAUDIO2:		return TEXT("XAUDIO2");
-			case NGS2:			return TEXT("NGS2");
-			case ALSA:			return TEXT("ALSA");
-			case CORE_AUDIO:	return TEXT("CORE_AUDIO");
-			case OPENAL:		return TEXT("OPENAL");
-			case HTML5:			return TEXT("HTML5");
-			case DUMMY:			return TEXT("DUMMY");
-			default:			check(false);
+				case WASAPI:		return TEXT("WASAPI");
+				case XAUDIO2:		return TEXT("XAUDIO2");
+				case NGS2:			return TEXT("NGS2");
+				case ALSA:			return TEXT("ALSA");
+				case CORE_AUDIO:	return TEXT("CORE_AUDIO");
+				case OPENAL:		return TEXT("OPENAL");
+				case HTML5:			return TEXT("HTML5");
+				case DUMMY:			return TEXT("DUMMY");
+				default:			check(false);
 			}
 			return TEXT("");
 		}
 	};
 
-	/**
-	* ESpeaker
-	* An enumeration to specify speaker types
-	*/
 	namespace ESpeaker
 	{
+		/** Values that represent speaker types. */
 		enum Type
 		{
 			FRONT_LEFT,
@@ -115,31 +132,30 @@ namespace UAudio
 			SPEAKER_TYPE_COUNT
 		};
 
-		/** @return the stringified version of the enum passed in */
 		static inline const TCHAR* ToString(ESpeaker::Type Speaker)
 		{
 			switch (Speaker)
 			{
-			case FRONT_LEFT:				return TEXT("FRONT_LEFT");
-			case FRONT_RIGHT:				return TEXT("FRONT_RIGHT");
-			case FRONT_CENTER:				return TEXT("FRONT_CENTER");
-			case LOW_FREQUENCY:				return TEXT("LOW_FREQUENCY");
-			case BACK_LEFT:					return TEXT("BACK_LEFT");
-			case BACK_RIGHT:				return TEXT("BACK_RIGHT");
-			case FRONT_LEFT_OF_CENTER:		return TEXT("FRONT_LEFT_OF_CENTER");
-			case FRONT_RIGHT_OF_CENTER:		return TEXT("FRONT_RIGHT_OF_CENTER");
-			case BACK_CENTER:				return TEXT("BACK_CENTER");
-			case SIDE_LEFT:					return TEXT("SIDE_LEFT");
-			case SIDE_RIGHT:				return TEXT("SIDE_RIGHT");
-			case TOP_CENTER:				return TEXT("TOP_CENTER");
-			case TOP_FRONT_LEFT:			return TEXT("TOP_FRONT_LEFT");
-			case TOP_FRONT_CENTER:			return TEXT("TOP_FRONT_CENTER");
-			case TOP_FRONT_RIGHT:			return TEXT("TOP_FRONT_RIGHT");
-			case TOP_BACK_LEFT:				return TEXT("TOP_BACK_LEFT");
-			case TOP_BACK_CENTER:			return TEXT("TOP_BACK_CENTER");
-			case TOP_BACK_RIGHT:			return TEXT("TOP_BACK_RIGHT");
-			case UNUSED:					return TEXT("UNUSED");
-			default:						return TEXT("UKNOWN");
+				case FRONT_LEFT:				return TEXT("FRONT_LEFT");
+				case FRONT_RIGHT:				return TEXT("FRONT_RIGHT");
+				case FRONT_CENTER:				return TEXT("FRONT_CENTER");
+				case LOW_FREQUENCY:				return TEXT("LOW_FREQUENCY");
+				case BACK_LEFT:					return TEXT("BACK_LEFT");
+				case BACK_RIGHT:				return TEXT("BACK_RIGHT");
+				case FRONT_LEFT_OF_CENTER:		return TEXT("FRONT_LEFT_OF_CENTER");
+				case FRONT_RIGHT_OF_CENTER:		return TEXT("FRONT_RIGHT_OF_CENTER");
+				case BACK_CENTER:				return TEXT("BACK_CENTER");
+				case SIDE_LEFT:					return TEXT("SIDE_LEFT");
+				case SIDE_RIGHT:				return TEXT("SIDE_RIGHT");
+				case TOP_CENTER:				return TEXT("TOP_CENTER");
+				case TOP_FRONT_LEFT:			return TEXT("TOP_FRONT_LEFT");
+				case TOP_FRONT_CENTER:			return TEXT("TOP_FRONT_CENTER");
+				case TOP_FRONT_RIGHT:			return TEXT("TOP_FRONT_RIGHT");
+				case TOP_BACK_LEFT:				return TEXT("TOP_BACK_LEFT");
+				case TOP_BACK_CENTER:			return TEXT("TOP_BACK_CENTER");
+				case TOP_BACK_RIGHT:			return TEXT("TOP_BACK_RIGHT");
+				case UNUSED:					return TEXT("UNUSED");
+				default:						return TEXT("UKNOWN");
 			}
 			return TEXT("");
 		}
@@ -148,5 +164,3 @@ namespace UAudio
 }
 
 
-#endif // #if ENABLE_UNREAL_AUDIO
-
diff --git a/Engine/Source/Runtime/UnrealAudio/Public/UnrealAudioVoice.h b/Engine/Source/Runtime/UnrealAudio/Public/UnrealAudioVoice.h
new file mode 100644
--- /dev/null
+++ b/Engine/Source/Runtime/UnrealAudio/Public/UnrealAudioVoice.h
@@ -0,0 +1,339 @@
+// Copyright 1998-2015 Epic Games, Inc. All Rights Reserved.
+
+#pragma once 
+
+#include "UnrealAudioModule.h"
+#include "UnrealAudioTypes.h"
+
+namespace UAudio
+{
+	namespace EVoiceError
+	{
+		enum Type
+		{
+			NONE = 0,
+			INVALID_HANDLE,
+			VOICE_NOT_PLAYING,
+			ALREADY_INITIALIZED,
+			NOT_INITIALIZED,
+			INVALID_ARGUMENTS,
+			LISTENER_ALREADY_ADDED,
+			LISTENER_NOT_ADDED,
+			UNKNOWN,
+		};
+
+		inline const TCHAR* ToString(EVoiceError::Type VoiceError)
+		{
+			switch (VoiceError)
+			{
+				case NONE:					return TEXT("NONE");
+
+				default:
+				case UNKNOWN:				return TEXT("UNKNOWN");
+			}
+		}
+	}
+
+	namespace EVoiceState
+	{
+		enum Type
+		{
+			UNINITIALIZED = 0,
+			STOPPED,
+			PAUSED,
+			PLAYING,
+			STOPPING,
+			PAUSING,
+			HAS_ERROR
+		};
+
+		inline const TCHAR* ToString(EVoiceState::Type VoiceState)
+		{
+			switch (VoiceState)
+			{
+				case UNINITIALIZED:		return TEXT("UNINITIALIZED");
+				case STOPPED:			return TEXT("STOPPED");
+				case PAUSED:			return TEXT("PAUSED");
+				case PLAYING:			return TEXT("PLAYING");
+				case STOPPING:			return TEXT("STOPPING");
+				case PAUSING:			return TEXT("PAUSING");
+				case HAS_ERROR:			return TEXT("HAS_ERROR");
+			}
+		}
+	}
+
+	namespace EVoicePlayingState
+	{
+		enum Type
+		{
+			NOT_PLAYING = 0,
+			PLAYING_REAL,
+			PLAYING_VIRTUAL,
+			SUSPENDED,
+		};
+
+		inline const TCHAR* ToString(EVoicePlayingState::Type PlayingState)
+		{
+			switch (PlayingState)
+			{
+				case NOT_PLAYING:		return TEXT("NOT_PLAYING");
+				case PLAYING_REAL:		return TEXT("PLAYING_REAL");
+				case PLAYING_VIRTUAL:	return TEXT("PLAYING_VIRTUAL");
+				case SUSPENDED:			return TEXT("SUSPENDED");
+			}
+		}
+	}
+
+	/** A voice initialization parameters. */
+	struct FVoiceInitializationParams
+	{
+		/** Sound file to use with the voice. */
+		TSharedPtr<ISoundFile> SoundFile;
+
+		/** Optional emitter to use with the voice. */
+		TSharedPtr<IEmitter> Emitter;
+
+		/** Baseline volume scale to use with this voice. */
+		float BaselineVolumeScale;
+
+		/** Baseline pitch scale to use with this voice. */
+		float BaselinePitchScale;
+
+		/** A priority weight value, used to determine voice stealing. */
+		float PriorityWeight;
+
+		/** Whether or not to loop this sound. */
+		bool bIsLooping;
+
+		/** Default constructor. */
+		FVoiceInitializationParams()
+			: SoundFile(nullptr)
+			, Emitter(nullptr)
+			, BaselineVolumeScale(1.0f)
+			, BaselinePitchScale(1.0f)
+			, PriorityWeight(1.0f)
+			, bIsLooping(false)
+		{}
+
+		/**
+		 * Constructor.
+		 *
+		 * @param	Other	Copy constructor.
+		 */
+
+		FVoiceInitializationParams(const FVoiceInitializationParams& Other)
+			: SoundFile(Other.SoundFile)
+			, Emitter(Other.Emitter)
+			, BaselineVolumeScale(Other.BaselineVolumeScale)
+			, BaselinePitchScale(Other.BaselinePitchScale)
+			, PriorityWeight(Other.PriorityWeight)
+			, bIsLooping(Other.bIsLooping)
+		{}
+	};
+
+	/** An object interface for listening to voice events. */
+	class UNREALAUDIO_API IVoiceListener
+	{
+	public:
+
+		/**
+		 * Called when a voice finishes playing or is stopped.
+		 *
+		 * @param	Voice	The voice object that finished playing a sound.
+		 */
+
+		virtual void OnVoiceDone(class IVoice* Voice) {}
+
+		/**
+		 * Called when a voice becomes virtual (i.e. can't be played because to many real voices are playing)
+		 * 
+		 * @param	Voice	The voice that became virtual.
+		 */
+
+		virtual void OnVoiceVirtual(class IVoice* Voice) {}
+
+		/**
+		 * Called when a voice becomes real (after it was virtual)
+		 *
+		 * @param	Voice	The voice that became real.
+		 */
+
+		virtual void OnVoiceReal(class IVoice* Voice) {}
+
+
+		virtual void OnVoiceSuspend(class IVoice* Voice) {}
+
+	};
+
+	/** An IVoice interface */
+	class UNREALAUDIO_API IVoice
+	{
+	public:
+		/** Destructor. */
+		virtual ~IVoice() {}
+
+		/** Adds the given notifier to the voice. The notifier will receive callbacks on various voice events. */
+		virtual EVoiceError::Type AddVoiceListener(IVoiceListener* Listener) = 0;
+
+		/**
+		 * Removes the voice listener described by Listener.
+		 *
+		 * @param [in]	Listener	The IVoiceListener to remove.
+		 *
+		 * @return	An EVoiceError::Type.
+		 */
+
+		virtual EVoiceError::Type RemoveVoiceListener(IVoiceListener* Listener) = 0;
+
+		/**
+		 * Sets the volume scale for this voice using linear volume scale.
+		 *
+		 * @param	InVolumeLinear	Linear volume set this voice to.
+		 * @param	InFadeTimeSec 	The time to take to get to the given volume scale from its current
+		 * 							volume scale.
+		 *
+		 * @return	An EVoiceError::Type.
+		 */
+
+		virtual EVoiceError::Type SetVolumeScale(float InVolumeLinear, float InFadeTimeSec = 0.0f) = 0;
+
+		/**
+		 * Gets volume scale.
+		 *
+		 * @param [in,out]	OutVolumeScale	The volume scale.
+		 *
+		 * @return	An EVoiceError::Type.
+		 */
+
+		virtual EVoiceError::Type GetVolumeScale(float& OutVolumeScale) const = 0;
+
+		/**
+		 * Gets baseline volume scale.
+		 *
+		 * @param [out]	OutBaselineVolumeScale	The baseline volume scale.
+		 *
+		 * @return	An EVoiceError::Type.
+		 */
+
+		virtual EVoiceError::Type GetBaselineVolumeScale(float& OutBaselineVolumeScale) const = 0;
+
+		/**
+		 * Gets volume attenuation (due to 3d spatialization).
+		 *
+		 * @param [in,out]	OutAttenuation	The out attenuation.
+		 *
+		 * @return	The volume attenuation.
+		 */
+
+		virtual EVoiceError::Type GetVolumeAttenuation(float& OutAttenuation) const = 0;
+
+		/**
+		 * Gets the volume product.
+		 *
+		 * @param [out]	OutVolumeProduct	The volume product.
+		 *
+		 * @return	An EVoiceError::Type.
+		 */
+
+		virtual EVoiceError::Type GetVolumeProduct(float& OutVolumeProduct) const = 0;
+
+		/**
+		 * Sets the pitch scale for this voice using linear volume scale.
+		 *
+		 * @param	InPitchScale 	Pitch scale to set for this voice.
+		 * @param	InFadeTimeSec	The time to scale the pitch to the given pitch scale.
+		 *
+		 * @return	An EVoiceError::Type.
+		 */
+
+		virtual EVoiceError::Type SetPitchScale(float InPitchScale, float InFadeTimeSec = 0.0f) = 0;
+
+		/**
+		 * Gets dynamic pitch scale of this voice.
+		 *
+		 * @param [out]	OutPitchScale	The pitch scale.
+		 *
+		 * @return	An EVoiceError::Type.
+		 */
+
+		virtual EVoiceError::Type GetPitchScale(float& OutPitchScale) const = 0;
+
+		/**
+		 * Gets the baseline pitch scale of this voice.
+		 *
+		 * @param [out]	OutBaselinePitchScale	The out baseline pitch scale.
+		 *
+		 * @return	An EVoiceError::Type.
+		 */
+
+		virtual EVoiceError::Type GetBaselinePitchScale(float& OutBaselinePitchScale) const = 0;
+
+		/**
+		 * Gets the pitch product.
+		 *
+		 * @param [in,out]	OutPitchProduct	The out pitch product.
+		 *
+		 * @return	An EVoiceError::Type.
+		 */
+
+		virtual EVoiceError::Type GetPitchProduct(float& OutPitchProduct) const = 0;
+
+		/**
+		 * Plays this voice.
+		 *
+		 * @return	An EVoiceError::Type.
+		 */
+
+		virtual EVoiceError::Type Play() = 0;
+
+		/**
+		 * Pauses this voice.
+		 *
+		 * @param	InFadeTimeSec	(Optional) The amount of time to spend fading out this sound before pausing it.
+		 *
+		 * @return	An EVoiceError::Type.
+		 */
+
+		virtual EVoiceError::Type Pause(float InFadeTimeSec = 0.0f) = 0;
+
+		/**
+		 * Stops this voice.
+		 *
+		 * @param	InFadeTimeSec	The amount of time to spend fading out this sound.
+		 * 
+		 * @note The voice is automatically released (i.e. freed) when the voice stops.
+		 *
+		 * @return	An EVoiceError::Type.
+		 */
+
+		virtual EVoiceError::Type Stop(float InFadeTimeSec = 0.0f) = 0;
+
+		/**
+		 * Gets a duration.
+		 *
+		 * @param [out]	OutDurationSec	The out duration security.
+		 *
+		 * @return	The duration.
+		 */
+
+		virtual EVoiceError::Type GetDuration(float& OutDurationSec) = 0;
+
+		/**
+		 * Is valid.
+		 *
+		 * @param [out]	bOutIsValid	The out is valid.
+		 *
+		 * @return	An EVoiceError::Type.
+		 */
+
+		virtual EVoiceError::Type IsValid(bool& bOutIsValid) const = 0;
+
+
+		virtual EVoiceError::Type GetVoiceState(EVoiceState::Type& OutVoicestate) const = 0;
+		virtual EVoiceError::Type GetPlayingState(EVoicePlayingState::Type& OutPlayingState) const = 0;
+		virtual EVoiceError::Type IsPlaying(bool& OutIsplaying) const = 0;
+		virtual EVoiceError::Type GetId(uint32& OutVoiceId) const = 0;
+
+	};
+}
+
