Attempt to fix the build by enabling ENABLE_UNREAL_AUDIO on windows

diff --git a/Engine/Source/Runtime/UnrealAudio/Public/UnrealAudioTypes.h b/Engine/Source/Runtime/UnrealAudio/Public/UnrealAudioTypes.h
--- a/Engine/Source/Runtime/UnrealAudio/Public/UnrealAudioTypes.h
+++ b/Engine/Source/Runtime/UnrealAudio/Public/UnrealAudioTypes.h
@@ -4,7 +4,7 @@
 
 // Only enable unreal audio on windows or mac
 #if PLATFORM_WINDOWS
-#define ENABLE_UNREAL_AUDIO 0
+#define ENABLE_UNREAL_AUDIO 1
 #else
 #define ENABLE_UNREAL_AUDIO 0
 #endif
