Unreal Audio compliance with vc2014 and clang?

diff --git a/Engine/Source/Runtime/UnrealAudio/Private/UnrealAudioEmitterManager.cpp b/Engine/Source/Runtime/UnrealAudio/Private/UnrealAudioEmitterManager.cpp
--- a/Engine/Source/Runtime/UnrealAudio/Private/UnrealAudioEmitterManager.cpp
+++ b/Engine/Source/Runtime/UnrealAudio/Private/UnrealAudioEmitterManager.cpp
@@ -70,7 +70,7 @@ namespace UAudio
 		check(Command.NumArguments == 1);
 		check(Command.Arguments[0].DataType == ECommandData::UINT32);
 
-		FEmitterHandle EmitterHandle = Command.Arguments[0].Data.UnsignedInt32;
+		FEmitterHandle EmitterHandle = FEmitterHandle::Create(Command.Arguments[0].Data.UnsignedInt32);
 		uint32 EmitterIndex = EmitterHandle.GetIndex();
 
 		if (EmitterIndex < (uint32)EmitterData.Num())
@@ -91,7 +91,7 @@ namespace UAudio
 		check(Command.NumArguments == 1);
 		check(Command.Arguments[0].DataType == ECommandData::UINT32);
 
-		FEmitterHandle EmitterHandle = Command.Arguments[0].Data.UnsignedInt32;
+		FEmitterHandle EmitterHandle = FEmitterHandle::Create(Command.Arguments[0].Data.UnsignedInt32);
 		uint32 EmitterIndex = EmitterHandle.GetIndex();
 		ValidateEmitterEntry(EmitterHandle, EmitterIndex);
 
@@ -106,7 +106,7 @@ namespace UAudio
 		check(Command.Arguments[2].DataType == ECommandData::FLOAT_32);
 		check(Command.Arguments[3].DataType == ECommandData::FLOAT_32);
 
-		FEmitterHandle EmitterHandle = Command.Arguments[0].Data.UnsignedInt32;
+		FEmitterHandle EmitterHandle = FEmitterHandle::Create(Command.Arguments[0].Data.UnsignedInt32);
 
 		FVector Position;
 		Position.X = Command.Arguments[1].Data.Float32Val;
diff --git a/Engine/Source/Runtime/UnrealAudio/Private/UnrealAudioEntityManager.h b/Engine/Source/Runtime/UnrealAudio/Private/UnrealAudioEntityManager.h
--- a/Engine/Source/Runtime/UnrealAudio/Private/UnrealAudioEntityManager.h
+++ b/Engine/Source/Runtime/UnrealAudio/Private/UnrealAudioEntityManager.h
@@ -20,31 +20,17 @@ namespace UAudio
 
 	public:
 
-		/** Default constructor. */
-
-		FEntityHandle()
-			: Id(INDEX_NONE)
-		{}
-
 		/**
-		 * Constructor.
+		 * Static constructor.
 		 *
 		 * @param	InId	Identifier for the in.
 		 */
 
-		FEntityHandle(uint32 InId)
-			: Id(InId)
-		{}
-
-		/**
-		 * Constructor.
-		 *
-		 * @param	Other	The other.
-		 */
-
-		FEntityHandle(const FEntityHandle& Other)
-			: Id(Other.Id)
-		{}
+		static FEntityHandle Create(uint32 InId = INDEX_NONE)
+		{
+			FEntityHandle Result;
+			Result.Id = InId;
+		}
 
 		/**
 		 * Query if this object is initialized.
@@ -179,7 +165,7 @@ namespace UAudio
 
 		FEntityHandle CreateHandle(uint32 EntityIndex, uint8 EntityGeneration) const
 		{
-			return (EntityIndex | (EntityGeneration << ENTITY_INDEX_BITS));
+			return FEntityHandle::Create(EntityIndex | (EntityGeneration << ENTITY_INDEX_BITS));
 		}
 
 		/** Number of free entities. */
diff --git a/Engine/Source/Runtime/UnrealAudio/Private/UnrealAudioHandles.h b/Engine/Source/Runtime/UnrealAudio/Private/UnrealAudioHandles.h
--- a/Engine/Source/Runtime/UnrealAudio/Private/UnrealAudioHandles.h
+++ b/Engine/Source/Runtime/UnrealAudio/Private/UnrealAudioHandles.h
@@ -12,11 +12,7 @@ namespace UAudio
 	struct FSoundFileHandle : public FEntityHandle
 	{
 		FSoundFileHandle()
-			: FEntityHandle()
-		{}
-
-		FSoundFileHandle(const FSoundFileHandle& InEntityHandle)
-			: FEntityHandle(InEntityHandle)
+			: FEntityHandle(FEntityHandle::Create())
 		{}
 
 		FSoundFileHandle(const FEntityHandle& InEntityHandle)
@@ -27,11 +23,7 @@ namespace UAudio
 	struct FVoiceHandle : public FEntityHandle
 	{
 		FVoiceHandle()
-			: FEntityHandle()
-		{}
-
-		FVoiceHandle(const FVoiceHandle& InEntityHandle)
-			: FEntityHandle(InEntityHandle)
+			: FEntityHandle(FEntityHandle::Create())
 		{}
 
 		FVoiceHandle(const FEntityHandle& InEntityHandle)
@@ -42,7 +34,7 @@ namespace UAudio
 	struct FEmitterHandle : public FEntityHandle
 	{
 		FEmitterHandle()
-			: FEntityHandle()
+			: FEntityHandle(FEntityHandle::Create())
 		{}
 
 		FEmitterHandle(const FEntityHandle& InEntityHandle)
diff --git a/Engine/Source/Runtime/UnrealAudio/Private/UnrealAudioSoundFileReader.cpp b/Engine/Source/Runtime/UnrealAudio/Private/UnrealAudioSoundFileReader.cpp
--- a/Engine/Source/Runtime/UnrealAudio/Private/UnrealAudioSoundFileReader.cpp
+++ b/Engine/Source/Runtime/UnrealAudio/Private/UnrealAudioSoundFileReader.cpp
@@ -755,8 +755,8 @@ namespace UAudio
 
 		if ((Description.FormatFlags & ESoundFileFormat::MAJOR_FORMAT_MASK) == ESoundFileFormat::OGG)
 		{
-			int32 Result = SoundFileCommand(FileHandle, SET_ENCODING_QUALITY, &EncodingQuality, sizeof(double));
-			if (Result != 1)
+			int32 Result2 = SoundFileCommand(FileHandle, SET_ENCODING_QUALITY, &EncodingQuality, sizeof(double));
+			if (Result2 != 1)
 			{
 				FString StrErr = SoundFileStrError(FileHandle);
 				UE_LOG(LogUnrealAudio, Error, TEXT("Failed to set encoding quality: %s"), *StrErr);
diff --git a/Engine/Source/Runtime/UnrealAudio/Private/UnrealAudioUtilities.h b/Engine/Source/Runtime/UnrealAudio/Private/UnrealAudioUtilities.h
--- a/Engine/Source/Runtime/UnrealAudio/Private/UnrealAudioUtilities.h
+++ b/Engine/Source/Runtime/UnrealAudio/Private/UnrealAudioUtilities.h
@@ -255,10 +255,10 @@ namespace UAudio
 			uint64 UnsignedInt64;
 			int32 Int32Val;
 			int64 Int64Val;
-			struct
-			{
+//			struct
+//			{
 				FEntityHandle Handle;
-			};
+//			};
 		} Data;
 	};
 
