Missing return on ::Create

diff --git a/Engine/Source/Runtime/UnrealAudio/Private/UnrealAudioEntityManager.h b/Engine/Source/Runtime/UnrealAudio/Private/UnrealAudioEntityManager.h
--- a/Engine/Source/Runtime/UnrealAudio/Private/UnrealAudioEntityManager.h
+++ b/Engine/Source/Runtime/UnrealAudio/Private/UnrealAudioEntityManager.h
@@ -30,6 +30,7 @@ namespace UAudio
 		{
 			FEntityHandle Result;
 			Result.Id = InId;
+			return Result;
 		}
 
 		/**
