Fixing CIS build error: Removing #pragma once from UnrealAudioSoundFileManager implementation file.

diff --git a/Engine/Source/Runtime/UnrealAudio/Private/UnrealAudioSoundFileManager.cpp b/Engine/Source/Runtime/UnrealAudio/Private/UnrealAudioSoundFileManager.cpp
--- a/Engine/Source/Runtime/UnrealAudio/Private/UnrealAudioSoundFileManager.cpp
+++ b/Engine/Source/Runtime/UnrealAudio/Private/UnrealAudioSoundFileManager.cpp
@@ -1,7 +1,5 @@
 // Copyright 1998-2015 Epic Games, Inc. All Rights Reserved.
 
-#pragma once 
-
 #include "UnrealAudioPrivate.h"
 #include "UnrealAudioModule.h"
 #include "UnrealAudioTypes.h"
