Removing deprecation warnings (Build Watcher)

diff --git a/Engine/Source/Runtime/UnrealAudio/Private/Tests/UnrealAudioSystemTests.cpp b/Engine/Source/Runtime/UnrealAudio/Private/Tests/UnrealAudioSystemTests.cpp
--- a/Engine/Source/Runtime/UnrealAudio/Private/Tests/UnrealAudioSystemTests.cpp
+++ b/Engine/Source/Runtime/UnrealAudio/Private/Tests/UnrealAudioSystemTests.cpp
@@ -89,7 +89,7 @@ namespace UAudio
 			}
 			else
 			{
-				EmitterLink->Link(EmittersHead);
+				EmitterLink->LinkHead(EmittersHead);
 			}
 		}
 
@@ -134,7 +134,7 @@ namespace UAudio
 					UE_LOG(LogUnrealAudio, Display, TEXT("[%d] Create"), Emitter->GetId());
 					Emitter->SetPosition(GetRandomPosition());
 					TLinkedList<TSharedPtr<IEmitter>>* EmitterLink = new TLinkedList<TSharedPtr<IEmitter>>(Emitter);
-					EmitterLink->Link(EmittersHead);
+					EmitterLink->LinkHead(EmittersHead);
 				}
 			}
 
