Fix warnings building samples due to missing paths.

diff --git a/Engine/Source/Runtime/UnrealAudio/UnrealAudio.Build.cs b/Engine/Source/Runtime/UnrealAudio/UnrealAudio.Build.cs
--- a/Engine/Source/Runtime/UnrealAudio/UnrealAudio.Build.cs
+++ b/Engine/Source/Runtime/UnrealAudio/UnrealAudio.Build.cs
@@ -33,26 +33,29 @@ public class UnrealAudio : ModuleRules
 		);
 
 		// Libsndfile DLL
-		string LibSndFilePath = Target.UEThirdPartyBinariesDirectory + "libsndfile/";
-		if (Target.Platform == UnrealTargetPlatform.Win32)
+		if(!bUsePrecompiled)
 		{
-			LibSndFilePath += "Win32";
-			PublicLibraryPaths.Add(LibSndFilePath);
-			PublicAdditionalLibraries.Add("libsndfile-1.lib");
-			PublicDelayLoadDLLs.Add("libsndfile-1.dll");
-		}
-		else if (Target.Platform == UnrealTargetPlatform.Win64)
-		{
-			LibSndFilePath += "Win64";
-			PublicLibraryPaths.Add(LibSndFilePath);
-			PublicAdditionalLibraries.Add("libsndfile-1.lib");
-			PublicDelayLoadDLLs.Add("libsndfile-1.dll");
+			string LibSndFilePath = Target.UEThirdPartyBinariesDirectory + "libsndfile/";
+			if (Target.Platform == UnrealTargetPlatform.Win32)
+			{
+				LibSndFilePath += "Win32";
+				PublicLibraryPaths.Add(LibSndFilePath);
+				PublicAdditionalLibraries.Add("libsndfile-1.lib");
+				PublicDelayLoadDLLs.Add("libsndfile-1.dll");
+			}
+			else if (Target.Platform == UnrealTargetPlatform.Win64)
+			{
+				LibSndFilePath += "Win64";
+				PublicLibraryPaths.Add(LibSndFilePath);
+				PublicAdditionalLibraries.Add("libsndfile-1.lib");
+				PublicDelayLoadDLLs.Add("libsndfile-1.dll");
+			}
+// 			else if (Target.Platform == UnrealTargetPlatform.Mac)
+// 			{
+// 				LibSndFilePath += "Mac/libsndfile.1.dylib";
+// 				PublicAdditionalLibraries.Add(LibSndFilePath);
+// 			}
 		}
-// 		else if (Target.Platform == UnrealTargetPlatform.Mac)
-// 		{
-// 			LibSndFilePath += "Mac/libsndfile.1.dylib";
-// 			PublicAdditionalLibraries.Add(LibSndFilePath);
-// 		}
 
 		PrecompileForTargets = PrecompileTargetsType.Editor;
 	}
